Copying //UE4/Release-Staging-4.12 to //UE4/Dev-Main (Source: //UE4/Release-4.12 @ 2955635)

diff --git a/Engine/Source/Editor/LevelEditor/Private/LevelEditorActions.cpp b/Engine/Source/Editor/LevelEditor/Private/LevelEditorActions.cpp
--- a/Engine/Source/Editor/LevelEditor/Private/LevelEditorActions.cpp
+++ b/Engine/Source/Editor/LevelEditor/Private/LevelEditorActions.cpp
@@ -227,15 +227,19 @@ void FLevelEditorActionCallbacks::OpenRecentFile( int32 RecentFileIndex )
 	FMainMRUFavoritesList* RecentsAndFavorites = MainFrameModule.GetMRUFavoritesList();
 
 	// Save the name of the file we are attempting to load as VerifyFile/AskSaveChanges might rearrange the MRU list on us
-	const FString NewFilename = RecentsAndFavorites->GetMRUItem( RecentFileIndex );
+	const FString NewPackageName = RecentsAndFavorites->GetMRUItem( RecentFileIndex );
 	
 	if( RecentsAndFavorites->VerifyMRUFile( RecentFileIndex ) )
 	{
 		// Prompt the user to save any outstanding changes.
 		if( FEditorFileUtils::SaveDirtyPackages(true, true, false) )
 		{
-			// Load the requested level.
-			FEditorFileUtils::LoadMap( NewFilename );
+			FString NewFilename;
+			if (FPackageName::TryConvertLongPackageNameToFilename(NewPackageName, NewFilename, FPackageName::GetMapPackageExtension()))
+			{
+				// Load the requested level.
+				FEditorFileUtils::LoadMap(NewFilename);
+			}
 		}
 		else
 		{
@@ -250,18 +254,22 @@ void FLevelEditorActionCallbacks::OpenFavoriteFile( int32 FavoriteFileIndex )
 	IMainFrameModule& MainFrameModule = FModuleManager::LoadModuleChecked<IMainFrameModule>( "MainFrame" );
 	FMainMRUFavoritesList* MRUFavoritesList = MainFrameModule.GetMRUFavoritesList();
 
-	const FString FileName = MRUFavoritesList->GetFavoritesItem( FavoriteFileIndex );
+	const FString PackageName = MRUFavoritesList->GetFavoritesItem( FavoriteFileIndex );
 
 	if( MRUFavoritesList->VerifyFavoritesFile( FavoriteFileIndex ) )
 	{
 		// Prompt the user to save any outstanding changes
 		if( FEditorFileUtils::SaveDirtyPackages(true, true, false) )
 		{
-			// Load the requested level.
-			FEditorFileUtils::LoadMap( FileName );
+			FString FileName;
+			if (FPackageName::TryConvertLongPackageNameToFilename(PackageName, FileName, FPackageName::GetMapPackageExtension()))
+			{
+				// Load the requested level.
+				FEditorFileUtils::LoadMap(FileName);
+			}
 
 			// Move the item to the head of the list
-			MRUFavoritesList->MoveFavoritesItemToHead( FileName );
+			MRUFavoritesList->MoveFavoritesItemToHead(PackageName);
 		}
 		else
 		{
@@ -277,21 +285,23 @@ void FLevelEditorActionCallbacks::ToggleFavorite()
 	FMainMRUFavoritesList* MRUFavoritesList = MainFrameModule.GetMRUFavoritesList();
 	check( MRUFavoritesList );
 
+	const FString PackageName = GetWorld()->GetOutermost()->GetName();
+
 	FString MapFileName;
-	const bool bMapFileExists = FPackageName::DoesPackageExist(GetWorld()->GetOutermost()->GetName(), NULL, &MapFileName);
+	const bool bMapFileExists = FPackageName::DoesPackageExist(PackageName, NULL, &MapFileName);
 
 	// If the user clicked the toggle favorites button, the map file should exist, but double check to be safe.
 	if ( bMapFileExists )
 	{
 		// If the map was already favorited, remove it from the favorites
-		if ( MRUFavoritesList->ContainsFavoritesItem( MapFileName ) )
+		if ( MRUFavoritesList->ContainsFavoritesItem(PackageName) )
 		{
-			MRUFavoritesList->RemoveFavoritesItem( MapFileName );
+			MRUFavoritesList->RemoveFavoritesItem(PackageName);
 		}
 		// If the map was not already favorited, add it to the favorites
 		else
 		{
-			MRUFavoritesList->AddFavoritesItem( MapFileName );
+			MRUFavoritesList->AddFavoritesItem(PackageName);
 		}
 	}
 }
@@ -302,13 +312,13 @@ void FLevelEditorActionCallbacks::RemoveFavorite( int32 FavoriteFileIndex )
 	IMainFrameModule& MainFrameModule = FModuleManager::LoadModuleChecked<IMainFrameModule>( "MainFrame" );
 	FMainMRUFavoritesList* MRUFavoritesList = MainFrameModule.GetMRUFavoritesList();
 
-	const FString FileName = MRUFavoritesList->GetFavoritesItem( FavoriteFileIndex );
+	const FString PackageName = MRUFavoritesList->GetFavoritesItem( FavoriteFileIndex );
 
 	if( MRUFavoritesList->VerifyFavoritesFile( FavoriteFileIndex ) )
 	{
-		if ( MRUFavoritesList->ContainsFavoritesItem( FileName ) )
+		if ( MRUFavoritesList->ContainsFavoritesItem(PackageName) )
 		{
-			MRUFavoritesList->RemoveFavoritesItem( FileName );
+			MRUFavoritesList->RemoveFavoritesItem(PackageName);
 		}
 	}
 }
@@ -333,18 +343,17 @@ bool FLevelEditorActionCallbacks::ToggleFavorite_IsChecked()
 {
 	bool bIsChecked = false;
 
+	const FString PackageName = GetWorld()->GetOutermost()->GetName();
+
 	FString FileName;
-	const bool bMapFileExists = FPackageName::DoesPackageExist(GetWorld()->GetOutermost()->GetName(), NULL, &FileName);
+	const bool bMapFileExists = FPackageName::DoesPackageExist(PackageName, NULL, &FileName);
 	
 	// If the map exists, determine its state based on whether the map is already favorited or not
 	if ( bMapFileExists )
 	{
 		IMainFrameModule& MainFrameModule = FModuleManager::LoadModuleChecked<IMainFrameModule>( "MainFrame" );
 
-		const FString CleanedName = FPaths::ConvertRelativePathToFull(FileName);
-		const bool bCleanAlreadyFavorited = MainFrameModule.GetMRUFavoritesList()->ContainsFavoritesItem( CleanedName );
-		const bool bAlreadyFavorited = bCleanAlreadyFavorited || MainFrameModule.GetMRUFavoritesList()->ContainsFavoritesItem( FileName );
-		bIsChecked = bAlreadyFavorited;
+		bIsChecked = MainFrameModule.GetMRUFavoritesList()->ContainsFavoritesItem(PackageName);
 	}
 
 	return bIsChecked;
diff --git a/Engine/Source/Editor/LevelEditor/Private/LevelEditorMenu.cpp b/Engine/Source/Editor/LevelEditor/Private/LevelEditorMenu.cpp
--- a/Engine/Source/Editor/LevelEditor/Private/LevelEditorMenu.cpp
+++ b/Engine/Source/Editor/LevelEditor/Private/LevelEditorMenu.cpp
@@ -279,7 +279,7 @@ TSharedRef< SWidget > FLevelEditorMenu::MakeLevelEditorMenu( const TSharedPtr<FU
 
 							const FString CurRecent = MRUFavorites.GetMRUItem( CurRecentIndex );
 
-							const FText ToolTip = FText::Format( LOCTEXT( "RecentFileToolTip", "Opens recent file: {0}" ), FText::FromString( CurRecent ) );
+							const FText ToolTip = FText::Format( LOCTEXT( "RecentLevelToolTip", "Opens recent level: {0}" ), FText::FromString( CurRecent ) );
 							const FText Label = FText::FromString( FPaths::GetBaseFilename( CurRecent ) );
 
 							InMenuBuilder.AddMenuEntry( OpenRecentFile, NAME_None, Label, ToolTip );
diff --git a/Engine/Source/Editor/LevelEditor/Private/LevelEditorToolBar.cpp b/Engine/Source/Editor/LevelEditor/Private/LevelEditorToolBar.cpp
--- a/Engine/Source/Editor/LevelEditor/Private/LevelEditorToolBar.cpp
+++ b/Engine/Source/Editor/LevelEditor/Private/LevelEditorToolBar.cpp
@@ -2110,12 +2110,12 @@ TSharedRef< SWidget > FLevelEditorToolBar::GenerateCinematicsMenuContent( TShare
 	MenuBuilder.EndSection();
 
 	UWorld* World = LevelEditorWeakPtr.Pin()->GetWorld();
-	const bool bHasAnyMatineeActors = !!TActorIterator<AMatineeActor>(World) || !!TActorIterator<ALevelSequenceActor>(World);
+	const bool bHasAnyCinematicsActors = !!TActorIterator<AMatineeActor>(World) || !!TActorIterator<ALevelSequenceActor>(World);
 
-	//Add a heading to separate the existing matinees from the 'Add New Matinee Actor' button
-	MenuBuilder.BeginSection("LevelEditorExistingMatinee", LOCTEXT( "MatineeMenuCombo_ExistingHeading", "Edit Existing Matinee" ) );
+	//Add a heading to separate the existing cinematics from the 'Add New Cinematic Actor' button
+	MenuBuilder.BeginSection("LevelEditorExistingCinematic", LOCTEXT( "CinematicMenuCombo_ExistingHeading", "Edit Existing Cinematic" ) );
 	{
-		if( bHasAnyMatineeActors )
+		if( bHasAnyCinematicsActors )
 		{
 			MenuBuilder.AddWidget(MiniSceneOutliner, FText::GetEmpty(), true);
 		}
diff --git a/Engine/Source/Editor/LevelEditor/Private/LevelViewportLayout.cpp b/Engine/Source/Editor/LevelEditor/Private/LevelViewportLayout.cpp
--- a/Engine/Source/Editor/LevelEditor/Private/LevelViewportLayout.cpp
+++ b/Engine/Source/Editor/LevelEditor/Private/LevelViewportLayout.cpp
@@ -245,7 +245,8 @@ void FLevelViewportLayout::SaveCommonLayoutString( const FString& SpecificLayout
 	// Save all our data using the additional layout config
 	for (auto& Pair : Viewports)
 	{
-		FString ConfigName = SpecificLayoutString + Pair.Key.ToString();
+		// The Viewports map is keyed on the full config name, so no need to prepend the SpecificLayoutString
+		FString ConfigName = Pair.Key.ToString();
 
 		Pair.Value->SaveConfig(ConfigName);
 
diff --git a/Engine/Source/Editor/LevelEditor/Public/LevelViewportLayout.h b/Engine/Source/Editor/LevelEditor/Public/LevelViewportLayout.h
--- a/Engine/Source/Editor/LevelEditor/Public/LevelViewportLayout.h
+++ b/Engine/Source/Editor/LevelEditor/Public/LevelViewportLayout.h
@@ -258,7 +258,7 @@ protected:
 	/** Curve for animating from a "restored" state to a maximized state */
 	FCurveSequence MaximizeAnimation;
 
-	/** List of all of the viewports in this layout */
+	/** List of all of the viewports in this layout, keyed on their config key */
 	TMap< FName, TSharedPtr< IViewportLayoutEntity > > Viewports;
 	
 	/** The parent tab where this layout resides */
diff --git a/Engine/Source/Editor/MovieSceneTools/Private/CurveKeyAreas/FloatCurveKeyArea.cpp b/Engine/Source/Editor/MovieSceneTools/Private/CurveKeyAreas/FloatCurveKeyArea.cpp
--- a/Engine/Source/Editor/MovieSceneTools/Private/CurveKeyAreas/FloatCurveKeyArea.cpp
+++ b/Engine/Source/Editor/MovieSceneTools/Private/CurveKeyAreas/FloatCurveKeyArea.cpp
@@ -58,6 +58,11 @@ TArray<FKeyHandle> FFloatCurveKeyArea::AddKeyUnique(float Time, EMovieSceneKeyIn
 			CurrentKey.LeaveTangentWeight = KeyToCopy.LeaveTangentWeight;
 		}
 	}
+	else if ( IntermediateValue.IsSet() )
+	{
+		float Value = IntermediateValue.GetValue();
+		Curve->UpdateOrAddKey(Time,Value);
+	}
 
 	return AddedKeyHandles;
 }
diff --git a/Engine/Source/Editor/MovieSceneTools/Private/CurveKeyAreas/IntegralKeyArea.cpp b/Engine/Source/Editor/MovieSceneTools/Private/CurveKeyAreas/IntegralKeyArea.cpp
--- a/Engine/Source/Editor/MovieSceneTools/Private/CurveKeyAreas/IntegralKeyArea.cpp
+++ b/Engine/Source/Editor/MovieSceneTools/Private/CurveKeyAreas/IntegralKeyArea.cpp
@@ -27,6 +27,10 @@ TArray<FKeyHandle> FIntegralCurveKeyAreaBase::AddKeyUnique(float Time, EMovieSce
 		EvaluateAndAddKey(Time, TimeToCopyFrom, CurrentKey);
 		AddedKeyHandles.Add(CurrentKey);
 	}
+	else
+	{
+		SetKeyValue(Time);
+	}
 
 	return AddedKeyHandles;
 }
diff --git a/Engine/Source/Editor/MovieSceneTools/Public/IntegralKeyArea.h b/Engine/Source/Editor/MovieSceneTools/Public/IntegralKeyArea.h
--- a/Engine/Source/Editor/MovieSceneTools/Public/IntegralKeyArea.h
+++ b/Engine/Source/Editor/MovieSceneTools/Public/IntegralKeyArea.h
@@ -45,6 +45,7 @@ public:
 protected:
 
 	virtual void EvaluateAndAddKey(float Time, float TimeToCopyFrom, FKeyHandle& CurrentKey) = 0;
+	virtual void SetKeyValue(float Time) = 0;
 
 protected:
 
@@ -158,6 +159,16 @@ protected:
 		Curve.AddKey(Time, Value, CurrentKey);
 	}
 
+	virtual void SetKeyValue(float Time) override
+	{
+		if (IntermediateValue.IsSet())
+		{
+			int32 Value = IntermediateValue.GetValue();
+
+			Curve.UpdateOrAddKey(Time, Value);
+		}
+	}
+
 protected:
 
 	TOptional<IntegralType> IntermediateValue;
diff --git a/Engine/Source/Editor/Persona/Private/SAnimNotifyPanel.cpp b/Engine/Source/Editor/Persona/Private/SAnimNotifyPanel.cpp
--- a/Engine/Source/Editor/Persona/Private/SAnimNotifyPanel.cpp
+++ b/Engine/Source/Editor/Persona/Private/SAnimNotifyPanel.cpp
@@ -2377,6 +2377,11 @@ void SAnimNotifyTrack::FillNewNotifyMenu(FMenuBuilder& MenuBuilder, bool bIsRepl
 	{
 		for(UClass* Class : NativeNotifyClasses)
 		{
+			if (Class->HasAllClassFlags(CLASS_Abstract))
+			{
+				continue; // skip abstract classes
+			}
+
 			const FText LabelText = Class->GetDisplayNameText();
 			const FString Label = LabelText.ToString();
 
@@ -4785,7 +4790,7 @@ void SAnimNotifyPanel::OnGetNativeNotifyData(TArray<UClass*>& OutClasses, UClass
 	{
 		UClass* Class = *It;
 
-		if(Class->IsChildOf(NotifyOutermost) && Class->HasAllClassFlags(CLASS_Native) && !Class->HasAllClassFlags(CLASS_Abstract) && !Class->IsInBlueprint())
+		if(Class->IsChildOf(NotifyOutermost) && Class->HasAllClassFlags(CLASS_Native) && !Class->IsInBlueprint())
 		{
 			OutClasses.Add(Class);
 			// Form class name to search later
diff --git a/Engine/Source/Editor/Persona/Private/SSkeletonTree.cpp b/Engine/Source/Editor/Persona/Private/SSkeletonTree.cpp
--- a/Engine/Source/Editor/Persona/Private/SSkeletonTree.cpp
+++ b/Engine/Source/Editor/Persona/Private/SSkeletonTree.cpp
@@ -58,7 +58,7 @@ public:
 	};
 
 private:
-	virtual bool CanCreateClass(UClass* ObjectClass) const { return true; }
+	virtual bool CanCreateClass(UClass* ObjectClass, bool& bOmitSubObjs) const override { return true; }
 
 	virtual void ProcessConstructedObject( UObject* CreatedObject )
 	{
diff --git a/Engine/Source/Editor/PixelInspector/Private/PixelInspector.cpp b/Engine/Source/Editor/PixelInspector/Private/PixelInspector.cpp
--- a/Engine/Source/Editor/PixelInspector/Private/PixelInspector.cpp
+++ b/Engine/Source/Editor/PixelInspector/Private/PixelInspector.cpp
@@ -2,6 +2,7 @@
 
 #include "PixelInspectorPrivatePCH.h"
 #include "PixelInspector.h"
+#include "PixelInspectorStyle.h"
 
 #include "Editor.h"
 #include "EditorModeManager.h"
@@ -9,6 +10,7 @@
 #include "IDetailsView.h"
 #include "PixelInspectorDetailsCustomization.h"
 #include "LevelEditor.h"
+#include "SNumericEntryBox.h"
 
 #define PIXEL_INSPECTOR_REQUEST_TIMEOUT 10
 #define MINIMUM_TICK_BETWEEN_CREATE_REQUEST 10
@@ -19,10 +21,15 @@ namespace PixelInspector
 	SPixelInspector::SPixelInspector()		
 	{
 		DisplayResult = nullptr;
+		LastViewportInspectionPosition = FIntPoint(-1, -1);
+		LastViewportId = 0;
+		LastViewportInspectionSize = FIntPoint(1, 1);
 		bIsPixelInspectorEnable = false;
 
 		Buffer_FinalColor_RGB8[0] = nullptr;
 		Buffer_FinalColor_RGB8[1] = nullptr;
+		Buffer_SceneColor_Float[0] = nullptr;
+		Buffer_SceneColor_Float[1] = nullptr;
 		Buffer_HDR_Float[0] = nullptr;
 		Buffer_HDR_Float[1] = nullptr;
 		Buffer_Depth_Float[0] = nullptr;
@@ -57,7 +64,7 @@ namespace PixelInspector
 	
 	void SPixelInspector::OnApplicationPreInputKeyDownListener(const FKeyEvent& InKeyEvent)
 	{
-		if (InKeyEvent.GetKey() == EKeys::Escape && bIsPixelInspectorEnable)
+		if (InKeyEvent.GetKey() == EKeys::Escape && (bIsPixelInspectorEnable))
 		{
 			// disable the pixel inspector
 			bIsPixelInspectorEnable = false;
@@ -137,6 +144,16 @@ namespace PixelInspector
 	BEGIN_SLATE_FUNCTION_BUILD_OPTIMIZATION
 	void SPixelInspector::Construct(const FArguments& InArgs)
 	{
+		//Set the LastViewportId to point on the active viewport
+		FViewport *ActiveViewport = GEditor->GetActiveViewport();
+		for (FEditorViewportClient *EditorViewport : GEditor->AllViewportClients)
+		{
+			if (ActiveViewport == EditorViewport->Viewport && EditorViewport->ViewState.GetReference() != nullptr)
+			{
+				LastViewportId = EditorViewport->ViewState.GetReference()->GetViewKey();
+			}
+		}
+
 		TSharedPtr<SBox> InspectorBox;
 		//Create the PixelInspector UI
 		TSharedPtr<SVerticalBox> VerticalBox = SNew(SVerticalBox)
@@ -145,23 +162,114 @@ namespace PixelInspector
 		[
 			SNew(SHorizontalBox)
 			+ SHorizontalBox::Slot()
+			.Padding(0.0f, 3.0f, 0.0f, 3.0f)
 			.AutoWidth()
 			[
-				SNew(SCheckBox)
+				SNew(SButton)
 				.HAlign(HAlign_Center)
-				.OnCheckStateChanged(this, &SPixelInspector::HandleTogglePixelInspectorEnable)
-				.IsChecked(this, &SPixelInspector::IsPixelInspectorEnable)
+				.ToolTipText(this, &SPixelInspector::GetPixelInspectorEnableButtonTooltipText)
+				.OnClicked(this, &SPixelInspector::HandleTogglePixelInspectorEnableButton)
+				[
+					SNew(SImage)
+					.Image(this, &SPixelInspector::GetPixelInspectorEnableButtonBrush)
+				]
 			]
 			+ SHorizontalBox::Slot()
-			.FillWidth(1.0f)
-			.Padding(0.0f, 3.0f, 6.0f, 3.0f)
+			.Padding(6.0f, 3.0f, 0.0f, 3.0f)
 			.VAlign(VAlign_Center)
+			.AutoWidth()
 			[
 				SNew(STextBlock)
-				.Text(LOCTEXT("PixelInspector_EnableCheckbox", "Enable Pixel Inspector"))
+				.MinDesiredWidth(75)
+				.Text(this, &SPixelInspector::GetPixelInspectorEnableButtonText)
 			]
 		]
 		+ SVerticalBox::Slot()
+		.AutoHeight()
+		[
+			SNew(SHorizontalBox)
+			+ SHorizontalBox::Slot()
+			.AutoWidth()
+			.Padding(0.0f, 3.0f, 16.0f, 3.0f)
+			.VAlign(VAlign_Center)
+			[
+				SNew(STextBlock)
+				.MinDesiredWidth(75)
+				.Text(LOCTEXT("PixelInspector_ViewportIdValue", "Viewport Id"))
+			]
+			+ SHorizontalBox::Slot()
+			.AutoWidth()
+			.Padding(0.0f, 3.0f, 0.0f, 3.0f)
+			[
+				SNew(SNumericEntryBox<uint32>)
+				.IsEnabled(false)
+				.MinDesiredValueWidth(75)
+				.Value(this, &SPixelInspector::GetCurrentViewportId)
+			]
+		]
+		+ SVerticalBox::Slot()
+		.AutoHeight()
+		[
+			SNew(SHorizontalBox)
+			+ SHorizontalBox::Slot()
+			.AutoWidth()
+			.Padding(0.0f, 3.0f, 16.0f, 3.0f)
+			.VAlign(VAlign_Center)
+			[
+				SNew(STextBlock)
+				.MinDesiredWidth(75)
+				.Text(LOCTEXT("PixelInspector_ViewportCoordinate", "Coordinate"))
+				.ToolTipText(LOCTEXT("PixelInspector_ViewportCoordinateTooltip", "Coordinate relative to the inspected viewport"))
+			]
+			+ SHorizontalBox::Slot()
+			.AutoWidth()
+			.Padding(0.0f, 3.0f, 8.0f, 3.0f)
+			[
+				SNew(SNumericEntryBox<int32>)
+				.IsEnabled(this, &SPixelInspector::IsPixelInspectorEnable)
+				.Value(this, &SPixelInspector::GetCurrentCoordinateX)
+				.OnValueChanged(this, &SPixelInspector::SetCurrentCoordinateX)
+				.OnValueCommitted(this, &SPixelInspector::SetCurrentCoordinateXCommit)
+				.AllowSpin(true)
+				.MinValue(0)
+				.MaxSliderValue(this, &SPixelInspector::GetMaxCoordinateX)
+				.MinDesiredValueWidth(75)
+				.Label()
+				[
+					SNew(SBox)
+					.VAlign(VAlign_Center)
+					[
+						SNew(STextBlock)
+						.Text(LOCTEXT("CoordinateViewport_X", "X"))
+					]
+				]
+			]
+			+ SHorizontalBox::Slot()
+			.Padding(0.0f, 3.0f, 8.0f, 3.0f)
+			.AutoWidth()
+			[
+				SNew(SNumericEntryBox<int32>)
+				.IsEnabled(this, &SPixelInspector::IsPixelInspectorEnable)
+				.Value(this, &SPixelInspector::GetCurrentCoordinateY)
+				.OnValueChanged(this, &SPixelInspector::SetCurrentCoordinateY)
+				.OnValueCommitted(this, &SPixelInspector::SetCurrentCoordinateYCommit)
+				.AllowSpin(true)
+				.MinValue(0)
+				.MaxSliderValue(this, &SPixelInspector::GetMaxCoordinateY)
+				.MinDesiredValueWidth(75)
+				.Label()
+				[
+					SNew(SBox)
+					.VAlign(VAlign_Center)
+					[
+						SNew(STextBlock)
+						.Text(LOCTEXT("CoordinateViewport_Y", "Y"))
+					]
+				]
+			]
+		]
+		+ SVerticalBox::Slot()
+		.Padding(0.0f, 12.0f, 0.0f, 3.0f)
 		.FillHeight(1.0f)
 		[
 			SAssignNew(InspectorBox, SBox)
@@ -169,7 +277,13 @@ namespace PixelInspector
 
 		FPropertyEditorModule& PropertyEditorModule = FModuleManager::GetModuleChecked<FPropertyEditorModule>("PropertyEditor");
 		FDetailsViewArgs DetailsViewArgs;
-		DetailsViewArgs.bAllowSearch = true;
+		DetailsViewArgs.bAllowSearch = false;
+		DetailsViewArgs.bLockable = false;
+		DetailsViewArgs.bShowActorLabel = false;
+		DetailsViewArgs.bShowOptions = false;
+		DetailsViewArgs.bUpdatesFromSelection = false;
+		DetailsViewArgs.bHideSelectionTip = true;
+		DetailsViewArgs.bSearchInitialKeyFocus = false;
 		DetailsViewArgs.NameAreaSettings = FDetailsViewArgs::HideNameArea;
 		DisplayDetailsView = PropertyEditorModule.CreateDetailView(DetailsViewArgs);
 		InspectorBox->SetContent(DisplayDetailsView->AsShared());
@@ -191,35 +305,129 @@ namespace PixelInspector
 		
 	}
 
+	END_SLATE_FUNCTION_BUILD_OPTIMIZATION
 
-	void SPixelInspector::Tick(const FGeometry& AllottedGeometry, const double InCurrentTime, const float InDeltaTime)
+	FReply SPixelInspector::HandleTogglePixelInspectorEnableButton()
 	{
-		SCompoundWidget::Tick(AllottedGeometry, InCurrentTime, InDeltaTime);
-/*		if (!GIsRequestingExit)
+		bIsPixelInspectorEnable = !bIsPixelInspectorEnable;
+		if (bIsPixelInspectorEnable)
 		{
-			//Read back the buffer that are ready
-			ReadBackRequestData();
-		}*/
-		TickSinceLastCreateRequest++;
+			if (LastViewportInspectionPosition == FIntPoint(-1, -1))
+			{
+				//Let the system inspect a pixel so the user can see the UI appear
+				LastViewportInspectionPosition = FIntPoint(0, 0);
+			}
+			//Make sure the viewport is switch to realtime
+			SetCurrentViewportInRealtime();
+		}
+		return FReply::Handled();
 	}
 
-	END_SLATE_FUNCTION_BUILD_OPTIMIZATION
+	FText SPixelInspector::GetPixelInspectorEnableButtonText() const
+	{
+		if (bIsPixelInspectorEnable)
+		{
+			return LOCTEXT("PixelInspector_EnableCheckbox_Inspecting", "Inspecting");
+		}
 
-	ECheckBoxState SPixelInspector::IsPixelInspectorEnable() const
+		return LOCTEXT("PixelInspectorMouseHover_EnableCheckbox", "Start Pixel Inspector");
+	}
+
+	FText SPixelInspector::GetPixelInspectorEnableButtonTooltipText() const
 	{
-		return bIsPixelInspectorEnable ? ECheckBoxState::Checked : ECheckBoxState::Unchecked;
+		if (bIsPixelInspectorEnable)
+		{
+			return LOCTEXT("PixelInspector_EnableCheckbox_ESC", "Inspecting (ESC to stop)");
+		}
+
+		return LOCTEXT("PixelInspectorMouseHover_EnableCheckbox", "Start Pixel Inspection");
+	}
+
+	const FSlateBrush* SPixelInspector::GetPixelInspectorEnableButtonBrush() const
+	{
+		return bIsPixelInspectorEnable ? FPixelInspectorStyle::Get()->GetBrush("PixelInspector.Enabled") : FPixelInspectorStyle::Get()->GetBrush("PixelInspector.Disabled");
+	}
+
+	void SPixelInspector::SetCurrentCoordinateXCommit(int32 NewValue, ETextCommit::Type)
+	{
+		ReleaseAllRequests();
+		SetCurrentCoordinateX(NewValue);
+	}
+
+	void SPixelInspector::SetCurrentCoordinateX(int32 NewValue)
+	{
+		LastViewportInspectionPosition.X = NewValue;
+	}
+
+	void SPixelInspector::SetCurrentCoordinateYCommit(int32 NewValue, ETextCommit::Type)
+	{
+		ReleaseAllRequests();
+		SetCurrentCoordinateY(NewValue);
+	}
+	void SPixelInspector::SetCurrentCoordinateY(int32 NewValue)
+	{
+		LastViewportInspectionPosition.Y = NewValue;
+	}
+
+	void SPixelInspector::SetCurrentCoordinate(FIntPoint NewCoordinate, bool ReleaseAllRequest)
+	{
+		if (ReleaseAllRequest)
+		{
+			ReleaseAllRequests();
+		}
+		LastViewportInspectionPosition.X = NewCoordinate.X;
+		LastViewportInspectionPosition.Y = NewCoordinate.Y;
 	}
 
-	void SPixelInspector::HandleTogglePixelInspectorEnable(ECheckBoxState CheckType)
+	TOptional<int32> SPixelInspector::GetMaxCoordinateX() const
 	{
-		bIsPixelInspectorEnable = (CheckType == ECheckBoxState::Checked);
+		return LastViewportInspectionSize.X - 1;
+	}
+
+	TOptional<int32> SPixelInspector::GetMaxCoordinateY() const
+	{
+		return LastViewportInspectionSize.Y - 1;
+	}
+
+	void SPixelInspector::SetCurrentViewportInRealtime()
+	{
+		//Force viewport refresh
+		for (FEditorViewportClient *EditorViewport : GEditor->AllViewportClients)
+		{
+			if (EditorViewport->ViewState.GetReference() != nullptr)
+			{
+				if (EditorViewport->ViewState.GetReference()->GetViewKey() == LastViewportId)
+				{
+					if (!EditorViewport->IsRealtime())
+					{
+						EditorViewport->SetRealtime(true);
+					}
+				}
+			}
+		}
 	}
 
+	void SPixelInspector::Tick(const FGeometry& AllottedGeometry, const double InCurrentTime, const float InDeltaTime)
+	{
+		SCompoundWidget::Tick(AllottedGeometry, InCurrentTime, InDeltaTime);
+		TickSinceLastCreateRequest++;
+	}
+	
 	void SPixelInspector::CreatePixelInspectorRequest(FIntPoint ScreenPosition, int32 viewportUniqueId, FSceneInterface *SceneInterface)
 	{
 		if (TickSinceLastCreateRequest < MINIMUM_TICK_BETWEEN_CREATE_REQUEST)
 			return;
 
+		if (ScreenPosition == FIntPoint(-1, -1))
+		{
+			return;
+		}
+		//Make sure we dont get value outside the viewport size
+		if ( ScreenPosition.X >= LastViewportInspectionSize.X || ScreenPosition.Y >= LastViewportInspectionSize.Y )
+		{
+			return;
+		}
+
 		TickSinceLastCreateRequest = 0;
 		// We need to know if the GBuffer is in low, default or high precision buffer
 		const auto CVarGBufferFormat = IConsoleManager::Get().FindTConsoleVariableDataInt(TEXT("r.GBufferFormat"));
@@ -253,6 +461,12 @@ namespace PixelInspector
 			Buffer_FinalColor_RGB8[BufferIndex]->RemoveFromRoot();
 			Buffer_FinalColor_RGB8[BufferIndex] = nullptr;
 		}
+		if (Buffer_SceneColor_Float[BufferIndex] != nullptr)
+		{
+			Buffer_SceneColor_Float[BufferIndex]->ClearFlags(RF_Standalone);
+			Buffer_SceneColor_Float[BufferIndex]->RemoveFromRoot();
+			Buffer_SceneColor_Float[BufferIndex] = nullptr;
+		}
 		if (Buffer_HDR_Float[BufferIndex] != nullptr)
 		{
 			Buffer_HDR_Float[BufferIndex]->ClearFlags(RF_Standalone);
@@ -314,6 +528,7 @@ namespace PixelInspector
 		ReleaseBuffers(LastBufferIndex);
 
 		FTextureRenderTargetResource* FinalColorRenderTargetResource = nullptr;
+		FTextureRenderTargetResource* SceneColorRenderTargetResource = nullptr;
 		FTextureRenderTargetResource* HDRRenderTargetResource = nullptr;
 		FTextureRenderTargetResource* DepthRenderTargetResource = nullptr;
 		FTextureRenderTargetResource* BufferARenderTargetResource = nullptr;
@@ -327,6 +542,14 @@ namespace PixelInspector
 		Buffer_FinalColor_RGB8[LastBufferIndex]->UpdateResourceImmediate(true);
 		FinalColorRenderTargetResource = Buffer_FinalColor_RGB8[LastBufferIndex]->GameThread_GetRenderTargetResource();
 
+		//Scene color is in RGB8 format
+		Buffer_SceneColor_Float[LastBufferIndex] = NewObject<UTextureRenderTarget2D>(GetTransientPackage(), TEXT("PixelInspectorBufferSceneColorTarget"), RF_Standalone);
+		Buffer_SceneColor_Float[LastBufferIndex]->AddToRoot();
+		Buffer_SceneColor_Float[LastBufferIndex]->InitCustomFormat(1, 1, PF_FloatRGBA, true);
+		Buffer_SceneColor_Float[LastBufferIndex]->ClearColor = FLinearColor::Black;
+		Buffer_SceneColor_Float[LastBufferIndex]->UpdateResourceImmediate(true);
+		SceneColorRenderTargetResource = Buffer_SceneColor_Float[LastBufferIndex]->GameThread_GetRenderTargetResource();
+
 		//HDR is in float RGB format
 		Buffer_HDR_Float[LastBufferIndex] = NewObject<UTextureRenderTarget2D>(GetTransientPackage(), TEXT("PixelInspectorBufferHDRTarget"), RF_Standalone);
 		Buffer_HDR_Float[LastBufferIndex]->AddToRoot();
@@ -341,8 +564,8 @@ namespace PixelInspector
 		Buffer_Depth_Float[LastBufferIndex]->InitCustomFormat(1, 1, PF_DepthStencil, true);
 		Buffer_Depth_Float[LastBufferIndex]->ClearColor = FLinearColor::Black;
 		Buffer_Depth_Float[LastBufferIndex]->UpdateResourceImmediate(true);
-		DepthRenderTargetResource = Buffer_Depth_Float[LastBufferIndex]->GameThread_GetRenderTargetResource();
-*/
+		DepthRenderTargetResource = Buffer_Depth_Float[LastBufferIndex]->GameThread_GetRenderTargetResource();*/
+
 
 		//Low precision GBuffer
 		if (GBufferFormat == 0)
@@ -398,7 +621,7 @@ namespace PixelInspector
 			BufferBCDERenderTargetResource = Buffer_BCDE_Float[LastBufferIndex]->GameThread_GetRenderTargetResource();
 		}
 		
-		SceneInterface->InitializePixelInspector(FinalColorRenderTargetResource, DepthRenderTargetResource, HDRRenderTargetResource, BufferARenderTargetResource, BufferBCDERenderTargetResource, LastBufferIndex);
+		SceneInterface->InitializePixelInspector(FinalColorRenderTargetResource, SceneColorRenderTargetResource, DepthRenderTargetResource, HDRRenderTargetResource, BufferARenderTargetResource, BufferBCDERenderTargetResource, LastBufferIndex);
 
 		return LastBufferIndex;
 	}
@@ -411,6 +634,10 @@ namespace PixelInspector
 			{
 				if (Requests[RequestIndex].FrameCountAfterRenderingCommandSend >= WAIT_FRAMENUMBER_BEFOREREADING)
 				{
+					if (Requests[RequestIndex].SourcePixelPosition == FIntPoint(-1, -1))
+					{
+						continue;
+					}
 					PixelInspectorResult PixelResult;
 					PixelResult.ScreenPosition = Requests[RequestIndex].SourcePixelPosition;
 					PixelResult.ViewUniqueId = Requests[RequestIndex].ViewId;
@@ -422,6 +649,14 @@ namespace PixelInspector
 						BufferFinalColorValue.Empty();
 					}
 					PixelResult.DecodeFinalColor(BufferFinalColorValue);
+
+					TArray<FLinearColor> BufferSceneColorValue;
+					FTextureRenderTargetResource* RTResourceSceneColor = Buffer_SceneColor_Float[Requests[RequestIndex].BufferIndex]->GameThread_GetRenderTargetResource();
+					if (RTResourceSceneColor->ReadLinearColorPixels(BufferSceneColorValue) == false)
+					{
+						BufferSceneColorValue.Empty();
+					}
+					PixelResult.DecodeSceneColor(BufferSceneColorValue);
 					
 					if (Buffer_Depth_Float[Requests[RequestIndex].BufferIndex] != nullptr)
 					{
@@ -530,6 +765,11 @@ namespace PixelInspector
 			}
 			DisplayResult->SetFromResult(AccumulationResult[0]);
 			DisplayDetailsView->SetObject(DisplayResult, true);
+			if (AccumulationResult[0].ScreenPosition != LastViewportInspectionPosition)
+			{
+				LastViewportInspectionPosition = AccumulationResult[0].ScreenPosition;
+			}
+			LastViewportId = AccumulationResult[0].ViewUniqueId;
 			AccumulationResult.RemoveAt(0);
 		}
 	}
diff --git a/Engine/Source/Editor/PixelInspector/Private/PixelInspector.h b/Engine/Source/Editor/PixelInspector/Private/PixelInspector.h
--- a/Engine/Source/Editor/PixelInspector/Private/PixelInspector.h
+++ b/Engine/Source/Editor/PixelInspector/Private/PixelInspector.h
@@ -48,10 +48,31 @@ namespace PixelInspector
 		void OnApplicationPreInputKeyDownListener(const FKeyEvent& InKeyEvent);
 
 		/** Button handlers */
-		ECheckBoxState IsPixelInspectorEnable() const;
-		void HandleTogglePixelInspectorEnable(ECheckBoxState CheckType);
+		FReply HandleTogglePixelInspectorEnableButton();
+		FText GetPixelInspectorEnableButtonText() const;
+		FText GetPixelInspectorEnableButtonTooltipText() const;
+		const FSlateBrush* GetPixelInspectorEnableButtonBrush() const;
+
+
+		TOptional<uint32> GetCurrentViewportId() const { return LastViewportId; }
+		FIntPoint GetCurrentCoordinate() const { return LastViewportInspectionPosition; }
+		TOptional<int32> GetCurrentCoordinateX() const { return LastViewportInspectionPosition.X; }
+		void SetCurrentCoordinateX(int32 NewValue);
+		void SetCurrentCoordinateXCommit(int32 NewValue, ETextCommit::Type);
+		TOptional<int32> GetCurrentCoordinateY() const { return LastViewportInspectionPosition.Y; }
+		void SetCurrentCoordinateY(int32 NewValue);
+		void SetCurrentCoordinateYCommit(int32 NewValue, ETextCommit::Type);
+		TOptional<int32> GetMaxCoordinateX() const;
+		TOptional<int32> GetMaxCoordinateY() const;
+
 		/** End button handlers */
 
+		bool IsPixelInspectorEnable() const { return bIsPixelInspectorEnable;	}
+
+		void SetCurrentCoordinate(FIntPoint NewCoordinate, bool ReleaseAllRequest);
+
+		void SetViewportInformation(int32 ViewportUniqueId, FIntPoint ViewportSize) { LastViewportInspectionSize = ViewportSize; }
+
 		/*
 		 * Create a request and the associate buffers
 		 *
@@ -88,7 +109,14 @@ namespace PixelInspector
 		void ReleaseBuffers(int32 BufferIndex);
 
 		void OnLevelActorDeleted(AActor* Actor);
+		
 		void OnRedrawViewport(bool bInvalidateHitProxies);
+
+		/* 
+		* Use by the Coordinate mode only, this change the realtime state of the viewport if the state is not true
+		*/
+		void SetCurrentViewportInRealtime();
+
 	private:
 
 		void ReleaseAllRequests();
@@ -109,6 +137,8 @@ namespace PixelInspector
 		UTextureRenderTarget2D* Buffer_FinalColor_RGB8[2];
 		//Depth Buffer
 		UTextureRenderTarget2D* Buffer_Depth_Float[2];
+		//SceneColor Buffer
+		UTextureRenderTarget2D* Buffer_SceneColor_Float[2];
 		//HDR Buffer
 		UTextureRenderTarget2D* Buffer_HDR_Float[2];
 		//GBufferA RenderTarget
@@ -127,7 +157,15 @@ namespace PixelInspector
 
 		//////////////////////////////////////////////////////////////////////////
 		// Display UObject to use the Detail Property Widget
+		
 		UPixelInspectorView *DisplayResult;
+		
+		FIntPoint LastViewportInspectionSize;
+		
+		FIntPoint LastViewportInspectionPosition;
+		
+		uint32 LastViewportId;
+
 		TSharedPtr<IDetailsView> DisplayDetailsView;
 	};
 }
diff --git a/Engine/Source/Editor/PixelInspector/Private/PixelInspectorDetailsCustomization.cpp b/Engine/Source/Editor/PixelInspector/Private/PixelInspectorDetailsCustomization.cpp
--- a/Engine/Source/Editor/PixelInspector/Private/PixelInspectorDetailsCustomization.cpp
+++ b/Engine/Source/Editor/PixelInspector/Private/PixelInspectorDetailsCustomization.cpp
@@ -8,6 +8,7 @@
 #include "DetailCategoryBuilder.h"
 #include "DetailWidgetRow.h"
 #include "PropertyHandle.h"
+#include "PixelInspectorModule.h"
 
 #define LOCTEXT_NAMESPACE "PixelInspector"
 
@@ -25,7 +26,7 @@ TSharedRef<IDetailCustomization> FPixelInspectorDetailsCustomization::MakeInstan
 TSharedRef<SHorizontalBox> FPixelInspectorDetailsCustomization::GetGridColorContext()
 {
 	TSharedRef<SHorizontalBox> HorizontalMainGrid = SNew(SHorizontalBox);
-	for (int RowIndex = 0; RowIndex < FinalColorContextGridSize; ++RowIndex)
+	for (int32 ColumnIndex = 0; ColumnIndex < FinalColorContextGridSize; ++ColumnIndex)
 	{
 		SBoxPanel::FSlot &HorizontalSlot = HorizontalMainGrid->AddSlot()
 			.AutoWidth()
@@ -34,7 +35,7 @@ TSharedRef<SHorizontalBox> FPixelInspectorDetailsCustomization::GetGridColorCont
 			.HAlign(HAlign_Center);
 
 		TSharedRef<SVerticalBox> VerticalColumn = SNew(SVerticalBox);
-		for (int ColumnIndex = 0; ColumnIndex < FinalColorContextGridSize; ++ColumnIndex)
+		for (int32 RowIndex = 0; RowIndex < FinalColorContextGridSize; ++RowIndex)
 		{
 			VerticalColumn->AddSlot()
 				.AutoHeight()
@@ -42,7 +43,7 @@ TSharedRef<SHorizontalBox> FPixelInspectorDetailsCustomization::GetGridColorCont
 				.VAlign(VAlign_Center)
 				.HAlign(HAlign_Center)
 				[
-					CreateColorCell(PixelInspectorView->FinalColorContext[RowIndex + ColumnIndex*FinalColorContextGridSize])
+					CreateColorCell(RowIndex, ColumnIndex, PixelInspectorView->FinalColorContext[ColumnIndex + RowIndex*FinalColorContextGridSize])
 				];
 		}
 		HorizontalSlot[VerticalColumn];
@@ -50,14 +51,52 @@ TSharedRef<SHorizontalBox> FPixelInspectorDetailsCustomization::GetGridColorCont
 	return HorizontalMainGrid;
 }
 
+FReply FPixelInspectorDetailsCustomization::HandleColorCellMouseButtonDown(const FGeometry &, const FPointerEvent &, int32 RowIndex, int32 ColumnIndex)
+{
+	//Send a create a request to the new Coordinate
+	if (RowIndex < 0 || RowIndex > FinalColorContextGridSize ||
+		ColumnIndex < 0 || ColumnIndex > FinalColorContextGridSize)
+		return FReply::Handled();
+
+	int32 DeltaX = ColumnIndex - FMath::FloorToInt((float)FinalColorContextGridSize / 2.0f);
+	int32 DeltaY = RowIndex - FMath::FloorToInt((float)FinalColorContextGridSize / 2.0f);
+	//If user click on the middle do nothing
+	if (DeltaX == 0 && DeltaY == 0)
+		return FReply::Handled();
+	//Get the PixelInspector module
+	FPixelInspectorModule* PixelInspectorModule = &FModuleManager::LoadModuleChecked<FPixelInspectorModule>(TEXT("PixelInspectorModule"));
+	if (PixelInspectorModule != nullptr)
+	{
+		FIntPoint InspectViewportPos = FIntPoint(-1, -1);
+		uint32 CoordinateViewportId = 0;
+		PixelInspectorModule->GetCoordinatePosition(InspectViewportPos, CoordinateViewportId);
+		if (InspectViewportPos == FIntPoint(-1, -1))
+			return FReply::Handled();
+
+		InspectViewportPos.X += DeltaX;
+		InspectViewportPos.Y += DeltaY;
+		if (InspectViewportPos.X < 0 || InspectViewportPos.Y < 0)
+			return FReply::Handled();
+
+		bool IsInspectorActive = PixelInspectorModule->IsPixelInspectorEnable();
+		if (!IsInspectorActive)
+		{
+			PixelInspectorModule->ActivateCoordinateMode();
+		}
+		PixelInspectorModule->SetCoordinatePosition(InspectViewportPos, true);
+	}
+	return FReply::Handled();
+}
 
-TSharedRef<SColorBlock> FPixelInspectorDetailsCustomization::CreateColorCell(const FLinearColor &CellColor)
+TSharedRef<SColorBlock> FPixelInspectorDetailsCustomization::CreateColorCell(int32 RowIndex, int32 ColumnIndex, const FLinearColor &CellColor)
 {
+	int32 SquareSize = FMath::FloorToInt(80.0f / (float)FinalColorContextGridSize);
 	return SNew(SColorBlock)
 		.Color(CellColor)
 		.ShowBackgroundForAlpha(false)
 		.IgnoreAlpha(true)
-		.Size(FVector2D(16.0f, 16.0f));
+		.Size(FVector2D(SquareSize, SquareSize))
+		.OnMouseButtonDown(this, &FPixelInspectorDetailsCustomization::HandleColorCellMouseButtonDown, RowIndex, ColumnIndex);
 }
 
 void FPixelInspectorDetailsCustomization::CustomizeDetails(IDetailLayoutBuilder& DetailBuilder)
@@ -70,7 +109,7 @@ void FPixelInspectorDetailsCustomization::CustomizeDetails(IDetailLayoutBuilder&
 	PixelInspectorView = Cast<UPixelInspectorView>(EditingObjects[0].Get());
 
 	IDetailCategoryBuilder& FinalColorCategory = DetailBuilder.EditCategory("FinalColor", FText::GetEmpty());
-	if (PixelInspectorView != nullptr)
+	if (PixelInspectorView.IsValid() )
 	{
 		FDetailWidgetRow& MergeRow = FinalColorCategory.AddCustomRow(LOCTEXT("FinalColorContextArray", "Context Color"))
 		.NameContent()
diff --git a/Engine/Source/Editor/PixelInspector/Private/PixelInspectorDetailsCustomization.h b/Engine/Source/Editor/PixelInspector/Private/PixelInspectorDetailsCustomization.h
--- a/Engine/Source/Editor/PixelInspector/Private/PixelInspectorDetailsCustomization.h
+++ b/Engine/Source/Editor/PixelInspector/Private/PixelInspectorDetailsCustomization.h
@@ -19,8 +19,10 @@ public:
 
 
 private:
+	FReply HandleColorCellMouseButtonDown(const FGeometry &, const FPointerEvent &, int32 RowIndex, int32 ColumnIndex);
+
 	TSharedRef<SHorizontalBox> GetGridColorContext();
-	TSharedRef<SColorBlock> CreateColorCell(const FLinearColor &CellColor);
+	TSharedRef<SColorBlock> CreateColorCell(int32 RowIndex, int32 ColumnIndex, const FLinearColor &CellColor);
 
 	/** Use MakeInstance to create an instance of this class */
 	FPixelInspectorDetailsCustomization();
diff --git a/Engine/Source/Editor/PixelInspector/Private/PixelInspectorModule.cpp b/Engine/Source/Editor/PixelInspector/Private/PixelInspectorModule.cpp
--- a/Engine/Source/Editor/PixelInspector/Private/PixelInspectorModule.cpp
+++ b/Engine/Source/Editor/PixelInspector/Private/PixelInspectorModule.cpp
@@ -26,9 +26,34 @@ TSharedRef<SWidget> FPixelInspectorModule::CreatePixelInspectorWidget()
 	return HPixelInspectorWindow->AsShared();
 }
 
+void FPixelInspectorModule::ActivateCoordinateMode()
+{
+	if (!HPixelInspectorWindow.IsValid())
+		return;
+	HPixelInspectorWindow->HandleTogglePixelInspectorEnableButton();
+}
+
 bool FPixelInspectorModule::IsPixelInspectorEnable()
 {
-	return HPixelInspectorWindow.IsValid() ? HPixelInspectorWindow->IsPixelInspectorEnable() == ECheckBoxState::Checked : false;
+	if (!HPixelInspectorWindow.IsValid())
+		return false;
+	return HPixelInspectorWindow->IsPixelInspectorEnable();
+}
+
+void FPixelInspectorModule::GetCoordinatePosition(FIntPoint &InspectViewportCoordinate, uint32 &InspectViewportId)
+{
+	InspectViewportCoordinate = FIntPoint(-1, -1);
+	if (!HPixelInspectorWindow.IsValid())
+		return;
+	InspectViewportId = HPixelInspectorWindow->GetCurrentViewportId().GetValue();
+	InspectViewportCoordinate = HPixelInspectorWindow->GetCurrentCoordinate();
+}
+
+void FPixelInspectorModule::SetCoordinatePosition(FIntPoint &InspectViewportCoordinate, bool ReleaseAllRequest)
+{
+	if (!HPixelInspectorWindow.IsValid())
+		return;
+	HPixelInspectorWindow->SetCurrentCoordinate(InspectViewportCoordinate, ReleaseAllRequest);
 }
 
 bool FPixelInspectorModule::GetViewportRealtime(int32 ViewportUid, bool IsCurrentlyRealtime, bool IsMouseInsideViewport)
@@ -57,13 +82,22 @@ bool FPixelInspectorModule::GetViewportRealtime(int32 ViewportUid, bool IsCurren
 	return IsCurrentlyRealtime;
 }
 
-void FPixelInspectorModule::CreatePixelInspectorRequest(FIntPoint ScreenPosition, int32 viewportUniqueId, FSceneInterface *SceneInterface)
+void FPixelInspectorModule::CreatePixelInspectorRequest(FIntPoint ScreenPosition, int32 ViewportUniqueId, FSceneInterface *SceneInterface)
+{
+	if (HPixelInspectorWindow.IsValid() == false)
+	{
+		return;
+	}
+	HPixelInspectorWindow->CreatePixelInspectorRequest(ScreenPosition, ViewportUniqueId, SceneInterface);
+}
+
+void FPixelInspectorModule::SetViewportInformation(int32 ViewportUniqueId, FIntPoint ViewportSize)
 {
 	if (HPixelInspectorWindow.IsValid() == false)
 	{
 		return;
 	}
-	HPixelInspectorWindow->CreatePixelInspectorRequest(ScreenPosition, viewportUniqueId, SceneInterface);
+	HPixelInspectorWindow->SetViewportInformation(ViewportUniqueId, ViewportSize);
 }
 
 void FPixelInspectorModule::ReadBackSync()
diff --git a/Engine/Source/Editor/PixelInspector/Private/PixelInspectorResult.cpp b/Engine/Source/Editor/PixelInspector/Private/PixelInspectorResult.cpp
--- a/Engine/Source/Editor/PixelInspector/Private/PixelInspectorResult.cpp
+++ b/Engine/Source/Editor/PixelInspector/Private/PixelInspectorResult.cpp
@@ -28,6 +28,17 @@ namespace PixelInspector
 			FinalColor.Add(LinearColor);
 		}
 	}
+	void PixelInspectorResult::DecodeSceneColor(TArray<FLinearColor> &BufferSceneColorValue)
+	{
+		if (BufferSceneColorValue.Num() <= 0)
+		{
+			SceneColor = FLinearColor::Black;
+			return;
+		}
+		SceneColor = BufferSceneColorValue[0];
+		//Set the alpha to 1.0 as the default value
+		SceneColor.A = 1.0f;
+	}
 	void PixelInspectorResult::DecodeDepth(TArray<FLinearColor> &BufferDepthValue)
 	{
 		if (BufferDepthValue.Num() <= 0)
diff --git a/Engine/Source/Editor/PixelInspector/Private/PixelInspectorResult.h b/Engine/Source/Editor/PixelInspector/Private/PixelInspectorResult.h
--- a/Engine/Source/Editor/PixelInspector/Private/PixelInspectorResult.h
+++ b/Engine/Source/Editor/PixelInspector/Private/PixelInspectorResult.h
@@ -62,6 +62,10 @@ namespace PixelInspector
 		// Final color 3x3 grid
 		TArray<FLinearColor> FinalColor;
 
+		//////////////////////////////////////////////////////////////////////////
+		// Scene color
+		FLinearColor SceneColor;
+
 		//////////////////////////////////////////////////////////////////////////
 		// Depth and world position
 		float Depth;
@@ -116,6 +120,7 @@ namespace PixelInspector
 		float IrisDistance;
 
 		void DecodeFinalColor(TArray<FColor> &BufferFinalColorValue);
+		void DecodeSceneColor(TArray<FLinearColor> &BufferSceneColorValue);
 		void DecodeDepth(TArray<FLinearColor> &BufferDepthValue);
 		void DecodeHDR(TArray<FLinearColor> &BufferHDRValue);
 
diff --git a/Engine/Source/Editor/PixelInspector/Private/PixelInspectorStyle.cpp b/Engine/Source/Editor/PixelInspector/Private/PixelInspectorStyle.cpp
--- a/Engine/Source/Editor/PixelInspector/Private/PixelInspectorStyle.cpp
+++ b/Engine/Source/Editor/PixelInspector/Private/PixelInspectorStyle.cpp
@@ -54,6 +54,8 @@ void FPixelInspectorStyle::Initialize()
 	// PixelInspector Icons
 	{
 		StyleSet->Set( "PixelInspector.TabIcon", new IMAGE_BRUSH( "Icons/PixelInspector/icon_PixelInspector_tab_16x", Icon16x16 ) );
+		StyleSet->Set( "PixelInspector.Enabled", new IMAGE_BRUSH( "Icons/PixelInspector/icon_PixelInspector_Stop_16x", Icon16x16 ) );
+		StyleSet->Set( "PixelInspector.Disabled", new IMAGE_BRUSH( "Icons/PixelInspector/icon_PixelInspector_Start_16x", Icon16x16 ) );
 	}
 
 	FSlateStyleRegistry::RegisterSlateStyle( *StyleSet.Get() );
diff --git a/Engine/Source/Editor/PixelInspector/Private/PixelInspectorView.cpp b/Engine/Source/Editor/PixelInspector/Private/PixelInspectorView.cpp
--- a/Engine/Source/Editor/PixelInspector/Private/PixelInspectorView.cpp
+++ b/Engine/Source/Editor/PixelInspector/Private/PixelInspectorView.cpp
@@ -7,13 +7,12 @@
 UPixelInspectorView::UPixelInspectorView(const FObjectInitializer& ObjectInitializer)
 	: Super(ObjectInitializer)
 {
-	ViewportId = -1;
-	PixelScreenPosition = FIntPoint(-1, -1);
 	for (int i = 0; i < FinalColorContextGridSize*FinalColorContextGridSize; ++i)
 	{
 		FinalColorContext[i] = FLinearColor::Green;
 	}
 	FinalColor = FLinearColor::Green;
+	SceneColor = FLinearColor::Green;
 	Luminance = 0.0f;
 	HdrColor = FLinearColor::Black;
 	Normal = FVector(0.0f);
@@ -43,8 +42,6 @@ UPixelInspectorView::UPixelInspectorView(const FObjectInitializer& ObjectInitial
 
 void UPixelInspectorView::SetFromResult(PixelInspector::PixelInspectorResult &Result)
 {
-	ViewportId = Result.ViewUniqueId;
-	PixelScreenPosition = Result.ScreenPosition;
 	FinalColor = FLinearColor::Green;
 	//Take the center of the array, TODO display the context
 	for (int i = 0; i < FinalColorContextGridSize*FinalColorContextGridSize; ++i)
@@ -63,6 +60,7 @@ void UPixelInspectorView::SetFromResult(PixelInspector::PixelInspectorResult &Re
 		}
 	}
 	
+	SceneColor = Result.SceneColor;
 	Luminance = Result.HdrLuminance;
 	HdrColor = Result.HdrColor;
 	Normal = Result.Normal;
diff --git a/Engine/Source/Editor/PixelInspector/Private/PixelInspectorView.h b/Engine/Source/Editor/PixelInspector/Private/PixelInspectorView.h
--- a/Engine/Source/Editor/PixelInspector/Private/PixelInspectorView.h
+++ b/Engine/Source/Editor/PixelInspector/Private/PixelInspectorView.h
@@ -3,7 +3,7 @@
 #pragma once
 #include "PixelInspectorView.generated.h"
 
-#define FinalColorContextGridSize 5
+#define FinalColorContextGridSize 7
 
 UCLASS(HideCategories = Object, MinimalAPI)
 class UPixelInspectorView : public UObject
@@ -12,18 +12,14 @@ class UPixelInspectorView : public UObject
 
 	FLinearColor FinalColorContext[FinalColorContextGridSize*FinalColorContextGridSize];
 
-	/** The viewport id from where the pixel was inspect. */
-	UPROPERTY(VisibleAnywhere, category = Identification)
-	int32 ViewportId;
-
-	/** The pixel relative to the inspected viewport. */
-	UPROPERTY(VisibleAnywhere, category = Identification)
-	FIntPoint PixelScreenPosition;
-
 	/** Final RGBA 8bits Color after tone mapping, default value is black. */
 	UPROPERTY(VisibleAnywhere, category = FinalColor)
 	FLinearColor FinalColor;
 
+	/** HDR RGB Color. */
+	UPROPERTY(VisibleAnywhere, category = SceneColor)
+	FLinearColor SceneColor;
+
 	/** HDR Luminance. */
 	UPROPERTY(VisibleAnywhere, category = HDR)
 	float Luminance;
@@ -118,7 +114,6 @@ class UPixelInspectorView : public UObject
 	UPROPERTY(VisibleAnywhere, category = GBufferD)
 	float IrisDistance;
 
-
 	void SetFromResult(PixelInspector::PixelInspectorResult &Result);
 	/*
 	//////////////////////////////////////////////////////////////////////////
diff --git a/Engine/Source/Editor/PixelInspector/Public/PixelInspectorModule.h b/Engine/Source/Editor/PixelInspector/Public/PixelInspectorModule.h
--- a/Engine/Source/Editor/PixelInspector/Public/PixelInspectorModule.h
+++ b/Engine/Source/Editor/PixelInspector/Public/PixelInspectorModule.h
@@ -1,5 +1,7 @@
 // Copyright 1998-2016 Epic Games, Inc. All Rights Reserved.
 
+#pragma once
+
 #include "ModuleInterface.h"
 #include "UnrealEd.h"
 
@@ -11,6 +13,7 @@ namespace PixelInspector { class SPixelInspector; };
 */
 class FPixelInspectorModule : public IModuleInterface
 {
+
 public:
 	/**
 	* Called right after the module DLL has been loaded and the module object has been created
@@ -25,11 +28,19 @@ public:
 	/** Creates the HLOD Outliner widget */
 	virtual TSharedRef<SWidget> CreatePixelInspectorWidget();
 
+	virtual void ActivateCoordinateMode();
+
 	virtual bool IsPixelInspectorEnable();
 
+	virtual void GetCoordinatePosition(FIntPoint &InspectViewportCoordinate, uint32 &InspectViewportId);
+
+	virtual void SetCoordinatePosition(FIntPoint &InspectViewportCoordinate, bool ReleaseAllRequest);
+
 	virtual bool GetViewportRealtime(int32 ViewportUid, bool IsCurrentlyRealtime, bool IsMouseInsideViewport);
 
-	virtual void CreatePixelInspectorRequest(FIntPoint ScreenPosition, int32 viewportUniqueId, FSceneInterface *SceneInterface);
+	virtual void CreatePixelInspectorRequest(FIntPoint ScreenPosition, int32 ViewportUniqueId, FSceneInterface *SceneInterface);
+
+	virtual void SetViewportInformation(int32 ViewportUniqueId, FIntPoint ViewportSize);
 
 	virtual void ReadBackSync();
 
diff --git a/Engine/Source/Editor/PropertyEditor/Private/DetailCategoryBuilderImpl.cpp b/Engine/Source/Editor/PropertyEditor/Private/DetailCategoryBuilderImpl.cpp
--- a/Engine/Source/Editor/PropertyEditor/Private/DetailCategoryBuilderImpl.cpp
+++ b/Engine/Source/Editor/PropertyEditor/Private/DetailCategoryBuilderImpl.cpp
@@ -511,6 +511,12 @@ bool FDetailCategoryImpl::ContainsOnlyAdvanced() const
 	return !bFavoriteCategory && SimpleChildNodes.Num() == 0 && AdvancedChildNodes.Num() > 0;
 }
 
+void FDetailCategoryImpl::GetCategoryInformation(int32 &SimpleChildNum, int32 &AdvanceChildNum) const
+{
+	SimpleChildNum = SimpleChildNodes.Num();
+	AdvanceChildNum = AdvancedChildNodes.Num();
+}
+
 void FDetailCategoryImpl::SetDisplayName( FName InCategoryName, const FText& LocalizedNameOverride )
 {
 	if( !LocalizedNameOverride.IsEmpty() )
diff --git a/Engine/Source/Editor/PropertyEditor/Private/DetailCategoryBuilderImpl.h b/Engine/Source/Editor/PropertyEditor/Private/DetailCategoryBuilderImpl.h
--- a/Engine/Source/Editor/PropertyEditor/Private/DetailCategoryBuilderImpl.h
+++ b/Engine/Source/Editor/PropertyEditor/Private/DetailCategoryBuilderImpl.h
@@ -248,6 +248,9 @@ public:
 	/** @return true if this category only contains advanced properties */
 	bool ContainsOnlyAdvanced() const;
 
+	/** @return true if this category only contains advanced properties */
+	void GetCategoryInformation(int32 &SimpleChildNum, int32 &AdvanceChildNum) const;
+
 	/**
 	 * Called when the advanced dropdown button is clicked 
 	 */
diff --git a/Engine/Source/Editor/PropertyEditor/Private/DetailLayoutBuilderImpl.cpp b/Engine/Source/Editor/PropertyEditor/Private/DetailLayoutBuilderImpl.cpp
--- a/Engine/Source/Editor/PropertyEditor/Private/DetailLayoutBuilderImpl.cpp
+++ b/Engine/Source/Editor/PropertyEditor/Private/DetailLayoutBuilderImpl.cpp
@@ -119,6 +119,11 @@ FDetailCategoryImpl& FDetailLayoutBuilderImpl::DefaultCategory( FName CategoryNa
 	return *CategoryImpl;
 }
 
+bool FDetailLayoutBuilderImpl::HasCategory(FName CategoryName)
+{
+	return DefaultCategoryMap.Contains(CategoryName);
+}
+
 void FDetailLayoutBuilderImpl::BuildCategories( const FCategoryMap& CategoryMap, TArray< TSharedRef<FDetailCategoryImpl> >& OutSimpleCategories, TArray< TSharedRef<FDetailCategoryImpl> >& OutAdvancedCategories )
 {
 	for( FCategoryMap::TConstIterator It(CategoryMap); It; ++It )
diff --git a/Engine/Source/Editor/PropertyEditor/Private/DetailLayoutBuilderImpl.h b/Engine/Source/Editor/PropertyEditor/Private/DetailLayoutBuilderImpl.h
--- a/Engine/Source/Editor/PropertyEditor/Private/DetailLayoutBuilderImpl.h
+++ b/Engine/Source/Editor/PropertyEditor/Private/DetailLayoutBuilderImpl.h
@@ -29,6 +29,11 @@ public:
 	 */
 	FDetailCategoryImpl& DefaultCategory( FName CategoryName );
 
+	/**
+	* Return true if the category exist
+	*/
+	bool HasCategory(FName CategoryName);
+
 	/**
 	 * Generates the layout for this detail builder                                                              
 	 */
diff --git a/Engine/Source/Editor/PropertyEditor/Private/IDetailsViewPrivate.h b/Engine/Source/Editor/PropertyEditor/Private/IDetailsViewPrivate.h
--- a/Engine/Source/Editor/PropertyEditor/Private/IDetailsViewPrivate.h
+++ b/Engine/Source/Editor/PropertyEditor/Private/IDetailsViewPrivate.h
@@ -56,6 +56,18 @@ public:
 	/** Causes the details view to be refreshed (new widgets generated) with the current set of objects */
 	virtual void ForceRefresh() = 0;
 
+	/** Causes the details view to move the scroll offset (by item)
+	* @param DeltaOffset	We add this value to the current scroll offset if the result is in the scrolling range
+	*/
+	virtual void MoveScrollOffset(int32 DeltaOffset) = 0;
+
+	/* Return true if the category exist in the detail view
+	* @param CategoryName			The name of the category we need informations
+	* @param SimplePropertiesNum		The Number of basic properties contain in the category
+	* @param AdvancePropertiesNum	The Number of advance properties contain in the category
+	*/
+	virtual bool GetCategoryInfo(FName CategoryName, int32 &SimplePropertiesNum, int32 &AdvancePropertiesNum) { SimplePropertiesNum = 0; AdvancePropertiesNum = 0; return false; }
+
 	/**
 	* Saves the expansion state of a tree node
 	*
diff --git a/Engine/Source/Editor/PropertyEditor/Private/PropertyHandleImpl.cpp b/Engine/Source/Editor/PropertyEditor/Private/PropertyHandleImpl.cpp
--- a/Engine/Source/Editor/PropertyEditor/Private/PropertyHandleImpl.cpp
+++ b/Engine/Source/Editor/PropertyEditor/Private/PropertyHandleImpl.cpp
@@ -1854,6 +1854,20 @@ void FPropertyHandleBase::SetToolTipText( const FText& ToolTip )
 	}
 }
 
+int32 FPropertyHandleBase::GetNumPerObjectValues() const
+{
+	TSharedPtr<FPropertyNode> PropertyNode = Implementation->GetPropertyNode();
+	if (PropertyNode.IsValid() && PropertyNode->GetProperty())
+	{
+		FComplexPropertyNode* ComplexNode = PropertyNode->FindComplexParent();
+		if (ComplexNode)
+		{
+			return ComplexNode->GetInstancesNum();
+		}
+	}
+	return 0;
+}
+
 FPropertyAccess::Result FPropertyHandleBase::SetPerObjectValues( const TArray<FString>& InPerObjectValues, EPropertyValueSetFlags::Type Flags )
 {
 	TSharedPtr<FPropertyNode> PropertyNode = Implementation->GetPropertyNode();
@@ -1875,7 +1889,7 @@ FPropertyAccess::Result FPropertyHandleBase::SetPerObjectValues( const TArray<FS
 	return FPropertyAccess::Fail;
 }
 
-FPropertyAccess::Result FPropertyHandleBase::GetPerObjectValues( TArray<FString>& OutPerObjectValues )
+FPropertyAccess::Result FPropertyHandleBase::GetPerObjectValues( TArray<FString>& OutPerObjectValues ) const
 {
 	FPropertyAccess::Result Result = FPropertyAccess::Fail;
 	TSharedPtr<FPropertyNode> PropertyNode = Implementation->GetPropertyNode();
@@ -1889,20 +1903,19 @@ FPropertyAccess::Result FPropertyHandleBase::GetPerObjectValues( TArray<FString>
 
 		if( ReadAddresses.Num() > 0 )
 		{
-			// Create a list of object names.
-			TArray<FString> ObjectNames;
-			ObjectNames.Empty( ReadAddresses.Num() );
-
-			// Copy each object's object property name off into the name list.
+			// Copy each object's value into the value list
+			OutPerObjectValues.SetNum( ReadAddresses.Num(), /*bAllowShrinking*/false );
 			for ( int32 AddrIndex = 0 ; AddrIndex < ReadAddresses.Num() ; ++AddrIndex )
 			{
 				uint8* Address = ReadAddresses.GetAddress(AddrIndex);
-
-				new( OutPerObjectValues ) FString();
 				if( Address )
 				{
 					NodeProperty->ExportText_Direct(OutPerObjectValues[AddrIndex], Address, Address, nullptr, 0 );
 				}
+				else
+				{
+					OutPerObjectValues[AddrIndex].Reset();
+				}
 			}
 
 			Result = FPropertyAccess::Success;
@@ -1912,6 +1925,67 @@ FPropertyAccess::Result FPropertyHandleBase::GetPerObjectValues( TArray<FString>
 	return Result;
 }
 
+FPropertyAccess::Result FPropertyHandleBase::SetPerObjectValue( const int32 ObjectIndex, const FString& ObjectValue, EPropertyValueSetFlags::Type Flags )
+{
+	FPropertyAccess::Result Result = FPropertyAccess::Fail;
+
+	TSharedPtr<FPropertyNode> PropertyNode = Implementation->GetPropertyNode();
+	if (PropertyNode.IsValid() && PropertyNode->GetProperty())
+	{
+		Implementation->EnumerateObjectsToModify(PropertyNode.Get(), [&](const FObjectBaseAddress& ObjectToModify, const int32 ObjIndex, const int32 NumObjects) -> bool
+		{
+			if (ObjIndex == ObjectIndex)
+			{
+				TArray<FObjectBaseAddress> ObjectsToModify;
+				ObjectsToModify.Add(ObjectToModify);
+
+				TArray<FString> PerObjectValues;
+				PerObjectValues.Add(ObjectValue);
+
+				Implementation->ImportText(ObjectsToModify, PerObjectValues, PropertyNode.Get(), Flags);
+
+				Result = FPropertyAccess::Success;
+				return false; // End enumeration
+			}
+			return true;
+		});
+	}
+
+	return Result;
+}
+
+FPropertyAccess::Result FPropertyHandleBase::GetPerObjectValue( const int32 ObjectIndex, FString& OutObjectValue ) const
+{
+	FPropertyAccess::Result Result = FPropertyAccess::Fail;
+
+	TSharedPtr<FPropertyNode> PropertyNode = Implementation->GetPropertyNode();
+	if (PropertyNode.IsValid() && PropertyNode->GetProperty())
+	{
+		// Get a list of addresses for objects handled by the property window.
+		FReadAddressList ReadAddresses;
+		PropertyNode->GetReadAddress(!!PropertyNode->HasNodeFlags(EPropertyNodeFlags::SingleSelectOnly), ReadAddresses, false);
+
+		UProperty* NodeProperty = PropertyNode->GetProperty();
+
+		if (ReadAddresses.IsValidIndex(ObjectIndex))
+		{
+			uint8* Address = ReadAddresses.GetAddress(ObjectIndex);
+			if (Address)
+			{
+				NodeProperty->ExportText_Direct(OutObjectValue, Address, Address, nullptr, 0);
+			}
+			else
+			{
+				OutObjectValue.Reset();
+			}
+
+			Result = FPropertyAccess::Success;
+		}
+	}
+
+	return Result;
+}
+
 TArray<FName> GetValidEnumEntries(UProperty* Property, const UEnum* InEnum)
 {
 	TArray<FName> ValidEnumValues;
diff --git a/Engine/Source/Editor/PropertyEditor/Private/PropertyHandleImpl.h b/Engine/Source/Editor/PropertyEditor/Private/PropertyHandleImpl.h
--- a/Engine/Source/Editor/PropertyEditor/Private/PropertyHandleImpl.h
+++ b/Engine/Source/Editor/PropertyEditor/Private/PropertyHandleImpl.h
@@ -425,8 +425,11 @@ public:
 	virtual const FString* GetInstanceMetaData(const FName& Key) const override;
 	virtual FText GetToolTipText() const override;
 	virtual void SetToolTipText(const FText& ToolTip) override;
+	virtual int32 GetNumPerObjectValues() const override;
 	virtual FPropertyAccess::Result SetPerObjectValues( const TArray<FString>& InPerObjectValues,  EPropertyValueSetFlags::Type Flags = EPropertyValueSetFlags::DefaultFlags ) override;
-	virtual FPropertyAccess::Result GetPerObjectValues( TArray<FString>& OutPerObjectValues ) override;
+	virtual FPropertyAccess::Result GetPerObjectValues( TArray<FString>& OutPerObjectValues ) const override;
+	virtual FPropertyAccess::Result SetPerObjectValue( const int32 ObjectIndex, const FString& ObjectValue, EPropertyValueSetFlags::Type Flags = EPropertyValueSetFlags::DefaultFlags ) override;
+	virtual FPropertyAccess::Result GetPerObjectValue( const int32 ObjectIndex, FString& OutObjectValue ) const override;
 	virtual bool GeneratePossibleValues(TArray< TSharedPtr<FString> >& OutOptionStrings, TArray< FText >& OutToolTips, TArray<bool>& OutRestrictedItems) override;
 	virtual FPropertyAccess::Result SetObjectValueFromSelection() override;
 	virtual void NotifyPreChange() override;
diff --git a/Engine/Source/Editor/PropertyEditor/Private/PropertyNode.h b/Engine/Source/Editor/PropertyEditor/Private/PropertyNode.h
--- a/Engine/Source/Editor/PropertyEditor/Private/PropertyNode.h
+++ b/Engine/Source/Editor/PropertyEditor/Private/PropertyNode.h
@@ -113,6 +113,11 @@ public:
 		return (Pair.Object.IsValid() || Pair.bIsStruct) ? Pair.ReadAddress : NULL;
 	}
 
+	bool IsValidIndex( int32 Index ) const
+	{
+		return ReadAddresses.IsValidIndex(Index);
+	}
+
 	void Reset()
 	{
 		ReadAddresses.Reset();
@@ -146,6 +151,11 @@ public:
 	{
 		return ReadAddressListData->GetAddress( Index );
 	}
+	
+	bool IsValidIndex( int32 Index ) const
+	{
+		return ReadAddressListData->IsValidIndex( Index );
+	}
 
 	void Reset()
 	{
diff --git a/Engine/Source/Editor/PropertyEditor/Private/SDetailSingleItemRow.cpp b/Engine/Source/Editor/PropertyEditor/Private/SDetailSingleItemRow.cpp
--- a/Engine/Source/Editor/PropertyEditor/Private/SDetailSingleItemRow.cpp
+++ b/Engine/Source/Editor/PropertyEditor/Private/SDetailSingleItemRow.cpp
@@ -64,6 +64,28 @@ private:
 	TAttribute< TOptional<float> > MaxWidth;
 };
 
+namespace SDetailSingleItemRow_Helper
+{
+	//Get the node item number in case it is expand we have to recursively count all expanded children
+	void RecursivelyGetItemShow(TSharedRef<IDetailTreeNode> ParentItem, int32 &ItemShowNum)
+	{
+		if (ParentItem->GetVisibility() == ENodeVisibility::Visible)
+		{
+			ItemShowNum++;
+		}
+
+		if (ParentItem->ShouldBeExpanded())
+		{
+			TArray< TSharedRef<IDetailTreeNode> > Childrens;
+			ParentItem->GetChildren(Childrens);
+			for (TSharedRef<IDetailTreeNode> ItemChild : Childrens)
+			{
+				RecursivelyGetItemShow(ItemChild, ItemShowNum);
+			}
+		}
+	}
+}
+
 FReply SDetailSingleItemRow::OnFavoriteToggle()
 {
 	if (Customization->GetPropertyNode().IsValid() && Customization->GetPropertyNode()->CanDisplayFavorite())
@@ -72,6 +94,64 @@ FReply SDetailSingleItemRow::OnFavoriteToggle()
 		Customization->GetPropertyNode()->SetFavorite(toggle);
 		if (OwnerTreeNode.IsValid())
 		{
+			//////////////////////////////////////////////////////////////////////////
+			// Calculate properly the scrolling offset (by item) to make sure the mouse stay over the same property
+
+			//Get the node item number in case it is expand we have to recursively count all childrens
+			int32 ExpandSize = 0;
+			if (OwnerTreeNode.Pin()->ShouldBeExpanded())
+			{
+				SDetailSingleItemRow_Helper::RecursivelyGetItemShow(OwnerTreeNode.Pin().ToSharedRef(), ExpandSize);
+			}
+			else
+			{
+				//if the item is not expand count is 1
+				ExpandSize = 1;
+			}
+			
+			//Get the number of favorite child (simple and advance) to know if the favorite category will be create or remove
+			FString CategoryFavoritesName = TEXT("Favorites");
+			FName CatFavName = *CategoryFavoritesName;
+			int32 SimplePropertiesNum = 0;
+			int32 AdvancePropertiesNum = 0;
+			bool HasCategoryFavorite = OwnerTreeNode.Pin()->GetDetailsView().GetCategoryInfo(CatFavName, SimplePropertiesNum, AdvancePropertiesNum);
+
+			//Check if the property we toggle is an advance property
+			bool IsAdvanceProperty = Customization->GetPropertyNode()->HasNodeFlags(EPropertyNodeFlags::IsAdvanced) == 0 ? false : true;
+
+			//Compute the scrolling offset by item
+			int32 ScrollingOffsetAdd = ExpandSize;
+			int32 ScrollingOffsetRemove = -ExpandSize;
+			if (HasCategoryFavorite)
+			{
+				//Adding the advance button in a category add 1 item
+				ScrollingOffsetAdd += (IsAdvanceProperty && AdvancePropertiesNum == 0) ? 1 : 0;
+
+				if (IsAdvanceProperty && AdvancePropertiesNum == 1)
+				{
+					//Removing the advance button count as 1 item
+					ScrollingOffsetRemove -= 1;
+				}
+				if (AdvancePropertiesNum + SimplePropertiesNum == 1)
+				{
+					//Removing a full category count as 2 items
+					ScrollingOffsetRemove -= 2;
+				}
+			}
+			else
+			{
+				//Adding new category (2 items) adding advance button (1 item)
+				ScrollingOffsetAdd += IsAdvanceProperty ? 3 : 2;
+				
+				//We should never remove an item from favorite if there is no favorite category
+				//Set the remove offset to 0
+				ScrollingOffsetRemove = 0;
+			}
+
+			//Apply the calculated offset
+			OwnerTreeNode.Pin()->GetDetailsView().MoveScrollOffset(toggle ? ScrollingOffsetAdd : ScrollingOffsetRemove);
+
+			//Refresh the tree
 			OwnerTreeNode.Pin()->GetDetailsView().ForceRefresh();
 		}
 	}
diff --git a/Engine/Source/Editor/PropertyEditor/Private/SDetailsView.cpp b/Engine/Source/Editor/PropertyEditor/Private/SDetailsView.cpp
--- a/Engine/Source/Editor/PropertyEditor/Private/SDetailsView.cpp
+++ b/Engine/Source/Editor/PropertyEditor/Private/SDetailsView.cpp
@@ -291,6 +291,22 @@ void SDetailsView::ForceRefresh()
 	SetObjectArrayPrivate( NewObjectList );
 }
 
+void SDetailsView::MoveScrollOffset(int32 DeltaOffset)
+{
+	DetailTree->AddScrollOffset((float)DeltaOffset);
+}
+
+bool SDetailsView::GetCategoryInfo(FName CategoryName, int32 &SimplePropertiesNum, int32 &AdvancePropertiesNum)
+{
+	SimplePropertiesNum = 0;
+	AdvancePropertiesNum = 0;
+	bool CategoryExist = DetailLayout.IsValid() ? DetailLayout->HasCategory(CategoryName) : false;
+	if (CategoryExist)
+	{
+		DetailLayout->DefaultCategory(CategoryName).GetCategoryInformation(SimplePropertiesNum, AdvancePropertiesNum);
+	}
+	return CategoryExist;
+}
 
 void SDetailsView::SetObjects( const TArray<UObject*>& InObjects, bool bForceRefresh/* = false*/, bool bOverrideLock/* = false*/ )
 {
diff --git a/Engine/Source/Editor/PropertyEditor/Private/SDetailsView.h b/Engine/Source/Editor/PropertyEditor/Private/SDetailsView.h
--- a/Engine/Source/Editor/PropertyEditor/Private/SDetailsView.h
+++ b/Engine/Source/Editor/PropertyEditor/Private/SDetailsView.h
@@ -21,6 +21,12 @@ public:
 	/** Causes the details view to be refreshed (new widgets generated) with the current set of objects */
 	void ForceRefresh() override;
 
+	/** Move the scrolling offset (by item), but do not refresh the tree*/
+	void MoveScrollOffset(int32 DeltaOffset) override;
+
+	/* Return true if the favorite category is currently displayed */
+	bool GetCategoryInfo(FName CategoryName, int32 &SimplePropertiesNum, int32 &AdvancePropertiesNum) override;
+
 	/**
 	 * Constructs the property view widgets                   
 	 */
diff --git a/Engine/Source/Editor/PropertyEditor/Private/SDetailsViewBase.cpp b/Engine/Source/Editor/PropertyEditor/Private/SDetailsViewBase.cpp
--- a/Engine/Source/Editor/PropertyEditor/Private/SDetailsViewBase.cpp
+++ b/Engine/Source/Editor/PropertyEditor/Private/SDetailsViewBase.cpp
@@ -1023,9 +1023,11 @@ void SDetailsViewBase::UpdatePropertyMapRecursive(FPropertyNode& InNode, FDetail
 {
 	UProperty* ParentProperty = InNode.GetProperty();
 	UStructProperty* ParentStructProp = Cast<UStructProperty>(ParentProperty);
-
 	for (int32 ChildIndex = 0; ChildIndex < InNode.GetNumChildNodes(); ++ChildIndex)
 	{
+		//Use the original value for each child
+		bool LocalUpdateFavoriteSystemOnly = bUpdateFavoriteSystemOnly;
+
 		TSharedPtr<FPropertyNode> ChildNodePtr = InNode.GetChildNode(ChildIndex);
 		FPropertyNode& ChildNode = *ChildNodePtr;
 		UProperty* Property = ChildNode.GetProperty();
@@ -1037,11 +1039,11 @@ void SDetailsViewBase::UpdatePropertyMapRecursive(FPropertyNode& InNode, FDetail
 			{
 				// Currently object property nodes do not provide any useful information other than being a container for its children.  We do not draw anything for them.
 				// When we encounter object property nodes, add their children instead of adding them to the tree.
-				UpdatePropertyMapRecursive(ChildNode, InDetailLayout, CurCategory, ObjNode, bEnableFavoriteSystem, bUpdateFavoriteSystemOnly);
+				UpdatePropertyMapRecursive(ChildNode, InDetailLayout, CurCategory, ObjNode, bEnableFavoriteSystem, LocalUpdateFavoriteSystemOnly);
 			}
 			else if (CategoryNode)
 			{
-				if (!bUpdateFavoriteSystemOnly)
+				if (!LocalUpdateFavoriteSystemOnly)
 				{
 					FName InstanceName = NAME_None;
 					FName CategoryName = CurCategory;
@@ -1056,7 +1058,7 @@ void SDetailsViewBase::UpdatePropertyMapRecursive(FPropertyNode& InNode, FDetail
 				}
 
 				// For category nodes, we just set the current category and recurse through the children
-				UpdatePropertyMapRecursive(ChildNode, InDetailLayout, CategoryNode->GetCategoryName(), CurObjectNode, bEnableFavoriteSystem, bUpdateFavoriteSystemOnly);
+				UpdatePropertyMapRecursive(ChildNode, InDetailLayout, CategoryNode->GetCategoryName(), CurObjectNode, bEnableFavoriteSystem, LocalUpdateFavoriteSystemOnly);
 			}
 			else
 			{
@@ -1108,7 +1110,7 @@ void SDetailsViewBase::UpdatePropertyMapRecursive(FPropertyNode& InNode, FDetail
 				const bool bIsUserVisible = IsPropertyVisible(PropertyAndParent);
 
 				// Inners of customized in structs should not be taken into consideration for customizing.  They are not designed to be individually customized when their parent is already customized
-				if (!bIsChildOfCustomizedStruct && !bUpdateFavoriteSystemOnly)
+				if (!bIsChildOfCustomizedStruct && !LocalUpdateFavoriteSystemOnly)
 				{
 					// Add any object classes with properties so we can ask them for custom property layouts later
 					ClassesWithProperties.Add(Property->GetOwnerStruct());
@@ -1126,7 +1128,7 @@ void SDetailsViewBase::UpdatePropertyMapRecursive(FPropertyNode& InNode, FDetail
 				}
 
 				// Do not add children of customized in struct properties or arrays
-				if (!bIsChildOfCustomizedStruct && !bIsChildOfArray && !bUpdateFavoriteSystemOnly)
+				if (!bIsChildOfCustomizedStruct && !bIsChildOfArray && !LocalUpdateFavoriteSystemOnly)
 				{
 					// Get the class property map
 					FClassInstanceToPropertyMap& ClassInstanceMap = ClassToPropertyMap.FindOrAdd(Property->GetOwnerStruct()->GetFName());
@@ -1158,7 +1160,7 @@ void SDetailsViewBase::UpdatePropertyMapRecursive(FPropertyNode& InNode, FDetail
 						CategoryName = PropertyCatagoryName;
 					}
 
-					if (!bUpdateFavoriteSystemOnly)
+					if (!LocalUpdateFavoriteSystemOnly)
 					{
 						if (IsPropertyReadOnly(PropertyAndParent))
 						{
@@ -1176,11 +1178,10 @@ void SDetailsViewBase::UpdatePropertyMapRecursive(FPropertyNode& InNode, FDetail
 						if (bIsCustomizedStruct)
 						{
 							bCanDisplayFavorite = false;
-							//CustomizedStruct child are not categorize since they are under an object but we have to put them in favorite
-							UpdatePropertyMapRecursive(ChildNode, InDetailLayout, CurCategory, CurObjectNode, bEnableFavoriteSystem, true);
+							//CustomizedStruct child are not categorize since they are under an object but we have to put them in favorite category if the user want to favorite them
+							LocalUpdateFavoriteSystemOnly = true;
 						}
-
-						if (ChildNodePtr->IsFavorite())
+						else if (ChildNodePtr->IsFavorite())
 						{
 							//Find or create the favorite category, we have to duplicate favorite property row under this category
 							FString CategoryFavoritesName = TEXT("Favorites");
@@ -1197,7 +1198,7 @@ void SDetailsViewBase::UpdatePropertyMapRecursive(FPropertyNode& InNode, FDetail
 								RootInstanceName = RootObjectParent->GetObjectBaseClass()->GetFName();
 							}
 
-							if (bUpdateFavoriteSystemOnly)
+							if (LocalUpdateFavoriteSystemOnly)
 							{
 								if (IsPropertyReadOnly(PropertyAndParent))
 								{
@@ -1230,10 +1231,10 @@ void SDetailsViewBase::UpdatePropertyMapRecursive(FPropertyNode& InNode, FDetail
 					&&	bIsUserVisible // Properties must be allowed to be visible by a user if they are not then their children are not visible either
 					&& (!bIsStruct || bPushOutStructProps); //  Only recurse into struct properties if they are going to be displayed as standalone properties in categories instead of inside an expandable area inside a category
 
-				if (bRecurseIntoChildren || bUpdateFavoriteSystemOnly)
+				if (bRecurseIntoChildren || LocalUpdateFavoriteSystemOnly)
 				{
 					// Built in struct properties or children of arras 
-					UpdatePropertyMapRecursive(ChildNode, InDetailLayout, CurCategory, CurObjectNode, bEnableFavoriteSystem, bUpdateFavoriteSystemOnly);
+					UpdatePropertyMapRecursive(ChildNode, InDetailLayout, CurCategory, CurObjectNode, bEnableFavoriteSystem, LocalUpdateFavoriteSystemOnly);
 				}
 			}
 		}
diff --git a/Engine/Source/Editor/PropertyEditor/Private/SStructureDetailsView.h b/Engine/Source/Editor/PropertyEditor/Private/SStructureDetailsView.h
--- a/Engine/Source/Editor/PropertyEditor/Private/SStructureDetailsView.h
+++ b/Engine/Source/Editor/PropertyEditor/Private/SStructureDetailsView.h
@@ -77,6 +77,7 @@ public:
 	}
 
 	virtual void ForceRefresh() override;
+	virtual void MoveScrollOffset(int32 DeltaOffset) override {}
 	virtual void AddExternalRootPropertyNode(TSharedRef<FPropertyNode> ExternalRootNode) override {}
 
 public:
diff --git a/Engine/Source/Editor/PropertyEditor/Public/PropertyHandle.h b/Engine/Source/Editor/PropertyEditor/Public/PropertyHandle.h
--- a/Engine/Source/Editor/PropertyEditor/Public/PropertyHandle.h
+++ b/Engine/Source/Editor/PropertyEditor/Public/PropertyHandle.h
@@ -255,6 +255,11 @@ public:
 	 */
 	virtual FPropertyAccess::Result SetObjectValueFromSelection() = 0;
 
+	/**
+	 * Gets the number of objects that this handle is editing
+	 */
+	virtual int32 GetNumPerObjectValues() const = 0;
+
 	/**
 	 * Sets a unique value for each object this handle is editing
 	 *
@@ -265,7 +270,23 @@ public:
 	/**
 	 * Gets a unique value for each object this handle is editing
 	 */
-	virtual FPropertyAccess::Result GetPerObjectValues( TArray<FString>& OutPerObjectValues ) = 0;
+	virtual FPropertyAccess::Result GetPerObjectValues( TArray<FString>& OutPerObjectValues ) const = 0;
+
+	/**
+	 * Sets a value on the specified object that this handle is editing
+	 *
+	 * @param ObjectIndex		The index of the object to set the value of
+	 * @param ObjectValue		The value to set on the given object
+	 */
+	virtual FPropertyAccess::Result SetPerObjectValue( const int32 ObjectIndex, const FString& ObjectValue, EPropertyValueSetFlags::Type Flags = EPropertyValueSetFlags::DefaultFlags ) = 0;
+
+	/**
+	 * Gets a value for the specified object that this handle is editing
+	 *
+	 * @param ObjectIndex		The index of the object to get the value of
+	 * @param OutObjectValue	Filled with the value for this object
+	 */
+	virtual FPropertyAccess::Result GetPerObjectValue( const int32 ObjectIndex, FString& OutObjectValue ) const = 0;
 
 	/**
 	 * @return The index of this element in an array if it is in one.  INDEX_NONE otherwise                                                              
diff --git a/Engine/Source/Editor/SequenceRecorder/Private/ActorRecording.cpp b/Engine/Source/Editor/SequenceRecorder/Private/ActorRecording.cpp
--- a/Engine/Source/Editor/SequenceRecorder/Private/ActorRecording.cpp
+++ b/Engine/Source/Editor/SequenceRecorder/Private/ActorRecording.cpp
@@ -271,7 +271,7 @@ void UActorRecording::StartRecordingActorProperties(ULevelSequence* CurrentSeque
 		FString TemplateName = GetUniqueSpawnableName(MovieScene, Actor->GetName());
 
 		UClass* ActorClass = Actor->GetClass();
-		AActor* ObjectTemplate = DuplicateObject<AActor>(Actor, MovieScene, *TemplateName);
+		AActor* ObjectTemplate = CastChecked<AActor>(StaticDuplicateObject(Actor, MovieScene, *TemplateName, RF_AllFlags & ~(RF_Transient)));
 
 		if (ObjectTemplate)
 		{
diff --git a/Engine/Source/Editor/Sequencer/Private/SSequencer.cpp b/Engine/Source/Editor/Sequencer/Private/SSequencer.cpp
--- a/Engine/Source/Editor/Sequencer/Private/SSequencer.cpp
+++ b/Engine/Source/Editor/Sequencer/Private/SSequencer.cpp
@@ -582,7 +582,7 @@ void SSequencer::HandleLabelBrowserSelectionChanged(FString NewLabel, ESelectInf
 	}
 	else
 	{
-		SearchBox->SetText(FText::FromString(FString(TEXT("label:")) + NewLabel));
+		SearchBox->SetText(FText::FromString(NewLabel));
 	}
 }
 
@@ -1180,7 +1180,7 @@ void SSequencer::OnOutlinerSearchChanged( const FText& Filter )
 
 		if ( FilterString.StartsWith( TEXT( "label:" ) ) )
 		{
-			LabelBrowser->SetSelectedLabel( FilterString.RightChop( 6 ) );
+			LabelBrowser->SetSelectedLabel(FilterString);
 		}
 		else
 		{
diff --git a/Engine/Source/Editor/Sequencer/Private/SSequencerLabelBrowser.cpp b/Engine/Source/Editor/Sequencer/Private/SSequencerLabelBrowser.cpp
--- a/Engine/Source/Editor/Sequencer/Private/SSequencerLabelBrowser.cpp
+++ b/Engine/Source/Editor/Sequencer/Private/SSequencerLabelBrowser.cpp
@@ -32,7 +32,7 @@ void SSequencerLabelBrowser::Construct(const FArguments& InArgs, TSharedRef<FSeq
 			.OnGenerateRow(this, &SSequencerLabelBrowser::HandleLabelTreeViewGenerateRow)
 			.OnGetChildren(this, &SSequencerLabelBrowser::HandleLabelTreeViewGetChildren)
 			.OnSelectionChanged(this, &SSequencerLabelBrowser::HandleLabelTreeViewSelectionChanged)
-			.SelectionMode(ESelectionMode::Single)
+			.SelectionMode(ESelectionMode::Multi)
 			.TreeItemsSource(&LabelList)
 	];
 
@@ -45,20 +45,28 @@ END_SLATE_FUNCTION_BUILD_OPTIMIZATION
 
 void SSequencerLabelBrowser::SetSelectedLabel(const FString& Label)
 {
-	if (!Label.IsEmpty())
+	TArray<FString> SplitLabelStrings;
+	Label.ParseIntoArray(SplitLabelStrings, TEXT(" "));
+
+	TArray<FString> SelectedLabels;
+	for (FString LabelString : SplitLabelStrings)
 	{
-		for (auto& Node : LabelList)
+		if (LabelString.StartsWith(TEXT("label:")))
 		{
-			if (Node->Label == Label)
-			{
-				LabelTreeView->SetSelection(Node);
-
-				return;
-			}
+			 LabelString = LabelString.RightChop( 6 );
+			 SelectedLabels.Add(LabelString);
 		}
 	}
 
 	LabelTreeView->ClearSelection();
+
+	for (auto& Node : LabelList)
+	{
+		if (SelectedLabels.Contains(Node->Label))
+		{
+			LabelTreeView->SetItemSelection(Node, true);
+		}
+	}
 }
 
 
@@ -197,9 +205,26 @@ void SSequencerLabelBrowser::HandleLabelTreeViewGetChildren(TSharedPtr<FSequence
 
 void SSequencerLabelBrowser::HandleLabelTreeViewSelectionChanged(TSharedPtr<FSequencerLabelTreeNode> InItem, ESelectInfo::Type SelectInfo)
 {
+	FString NewLabel;
+
+	TArray<TSharedPtr<FSequencerLabelTreeNode>> SelectedItems;
+	LabelTreeView->GetSelectedItems(SelectedItems);
+	for (TSharedPtr<FSequencerLabelTreeNode> Item : SelectedItems)
+	{
+		if (!NewLabel.IsEmpty())
+		{
+			NewLabel += TEXT(" ");
+		}
+		
+		if (!Item->Label.IsEmpty())
+		{
+			NewLabel += FString(TEXT("label:")) + Item->Label;
+		}
+	}
+
 	OnSelectionChanged.ExecuteIfBound(
 		InItem.IsValid()
-			? InItem->Label
+			? NewLabel
 			: FString(),
 		SelectInfo
 	);
diff --git a/Engine/Source/Editor/Sequencer/Private/Sequencer.cpp b/Engine/Source/Editor/Sequencer/Private/Sequencer.cpp
--- a/Engine/Source/Editor/Sequencer/Private/Sequencer.cpp
+++ b/Engine/Source/Editor/Sequencer/Private/Sequencer.cpp
@@ -38,6 +38,7 @@
 #include "ISequencerSection.h"
 #include "MovieSceneSequenceInstance.h"
 #include "IKeyArea.h"
+#include "ISettingsModule.h"
 #include "IDetailsView.h"
 #include "SnappingUtils.h"
 #include "GenericCommands.h"
@@ -127,6 +128,15 @@ void FSequencer::InitSequencer(const FSequencerInitParams& InitParams, const TSh
 
 	Settings = USequencerSettingsContainer::GetOrCreate<USequencerSettings>(*InitParams.ViewParams.UniqueName);
 
+	if (ISettingsModule* SettingsModule = FModuleManager::GetModulePtr<ISettingsModule>("Settings"))
+	{
+		FName SettingsSectionName = *Settings->GetName();
+		FText SettingsDisplayName = FText::FromString(FName::NameToDisplayString(*Settings->GetName(), false));
+		FText SettingsDescription = FText::FromString("Configure the look and feel of the " + FName::NameToDisplayString(*Settings->GetName(), false));
+
+		SettingsModule->RegisterSettings("Editor", "ContentEditors", SettingsSectionName, SettingsDisplayName, SettingsDescription, Settings);
+	}
+
 	ToolkitHost = InitParams.ToolkitHost;
 
 	ScrubPosition = InitParams.ViewParams.InitialScrubPosition;
@@ -227,6 +237,7 @@ FSequencer::FSequencer()
 	, bPerspectiveViewportCameraCutEnabled( false )
 	, bIsEditingWithinLevelEditor( false )
 	, bNeedTreeRefresh( false )
+	, bNeedInstanceRefresh( false )
 	, StoredPlaybackState( EMovieScenePlayerStatus::Stopped )
 	, bWasClosed( false )
 	, NodeTree( MakeShareable( new FSequencerNodeTree( *this ) ) )
@@ -295,6 +306,12 @@ void FSequencer::Close()
 		{
 			GLevelEditorModeTools().DeactivateMode(FSequencerEdMode::EM_SequencerMode);
 		}
+
+		if (ISettingsModule* SettingsModule = FModuleManager::GetModulePtr<ISettingsModule>("Settings"))
+		{
+			FName SettingsSectionName = *Settings->GetName();
+			SettingsModule->UnregisterSettings("Editor", "ContentEditors", SettingsSectionName);
+		}
 	}
 
 	if (bIsEditingWithinLevelEditor)
@@ -316,12 +333,16 @@ void FSequencer::Tick(float InDeltaTime)
 {
 	Selection.Tick();
 
-	if (bNeedTreeRefresh)
+	if ( bNeedInstanceRefresh )
 	{
-		SelectionPreview.Empty();
-
 		// @todo - Sequencer Will be called too often
 		UpdateRuntimeInstances();
+		bNeedInstanceRefresh = false;
+	}
+
+	if (bNeedTreeRefresh)
+	{
+		SelectionPreview.Empty();
 
 		SequencerWidget->UpdateLayoutTree();
 		bNeedTreeRefresh = false;
@@ -931,6 +952,7 @@ void FSequencer::NotifyMovieSceneDataChanged()
 	StoredPlaybackState = GetPlaybackStatus();
 	SetPlaybackStatus(EMovieScenePlayerStatus::Stopped);
 	bNeedTreeRefresh = true;
+	bNeedInstanceRefresh = true;
 
 	UpdatePlaybackRange();
 }
@@ -1256,7 +1278,7 @@ float FSequencer::GetGlobalTime() const
 }
 
 
-void FSequencer::SetGlobalTime( float NewTime, ESnapTimeMode SnapTimeMode )
+void FSequencer::SetGlobalTime( float NewTime, ESnapTimeMode SnapTimeMode, bool bLooped )
 {
 	if (IsAutoScrollEnabled())
 	{
@@ -1274,11 +1296,11 @@ void FSequencer::SetGlobalTime( float NewTime, ESnapTimeMode SnapTimeMode )
 		}
 	}
 
-	SetGlobalTimeDirectly(NewTime, SnapTimeMode);
+	SetGlobalTimeDirectly(NewTime, SnapTimeMode, bLooped);
 }
 
 
-void FSequencer::SetGlobalTimeDirectly( float NewTime, ESnapTimeMode SnapTimeMode )
+void FSequencer::SetGlobalTimeDirectly( float NewTime, ESnapTimeMode SnapTimeMode, bool bLooped )
 {
 	float LastTime = ScrubPosition;
 
@@ -1296,6 +1318,7 @@ void FSequencer::SetGlobalTimeDirectly( float NewTime, ESnapTimeMode SnapTimeMod
 	ScrubPosition = NewTime;
 
 	EMovieSceneUpdateData UpdateData(NewTime, LastTime);
+	UpdateData.bLooped = bLooped;
 	SequenceInstanceStack.Top()->Update(UpdateData, *this);
 
 	// Ensure any relevant spawned objects are cleaned up if we're playing back the master sequence
@@ -2301,6 +2324,7 @@ bool FSequencer::IsLooping() const
 
 void FSequencer::SetGlobalTimeLooped(float InTime)
 {
+	bool bLooped = false;
 	if (Settings->IsLooping())
 	{
 		const UMovieSceneSequence* FocusedSequence = GetFocusedMovieSceneSequence();
@@ -2308,7 +2332,8 @@ void FSequencer::SetGlobalTimeLooped(float InTime)
 		{
 			if (InTime > FocusedSequence->GetMovieScene()->GetPlaybackRange().GetUpperBoundValue())
 			{
-				InTime -= FocusedSequence->GetMovieScene()->GetPlaybackRange().Size<float>();
+				InTime = FocusedSequence->GetMovieScene()->GetPlaybackRange().GetLowerBoundValue();
+				bLooped = true;
 			}
 		}
 	}
@@ -2345,7 +2370,7 @@ void FSequencer::SetGlobalTimeLooped(float InTime)
 		}
 	}
 
-	SetGlobalTime(InTime);
+	SetGlobalTime(InTime, ESnapTimeMode::STM_None, bLooped);
 }
 
 
@@ -3051,12 +3076,12 @@ void FSequencer::OnMapOpened(const FString& Filename, bool bLoadAsTemplate)
 
 void FSequencer::OnLevelAdded(ULevel* InLevel, UWorld* InWorld)
 {
-	NotifyMovieSceneDataChanged();
+	bNeedInstanceRefresh = true;
 }
 
 void FSequencer::OnLevelRemoved(ULevel* InLevel, UWorld* InWorld)
 {
-	NotifyMovieSceneDataChanged();
+	bNeedInstanceRefresh = true;
 }
 
 void FSequencer::OnNewActorsDropped(const TArray<UObject*>& DroppedObjects, const TArray<AActor*>& DroppedActors)
@@ -3088,6 +3113,8 @@ void FSequencer::OnNewActorsDropped(const TArray<UObject*>& DroppedObjects, cons
 			FGuid PossessableGuid = CreateBinding(*NewActor, NewActor->GetActorLabel());
 			FGuid NewGuid = PossessableGuid;
 
+			OnActorAddedToSequencerEvent.Broadcast(NewActor, PossessableGuid);
+
 			if (bAddSpawnable)
 			{
 				FMovieSceneSpawnable* Spawnable = ConvertToSpawnableInternal(PossessableGuid);
@@ -3105,12 +3132,6 @@ void FSequencer::OnNewActorsDropped(const TArray<UObject*>& DroppedObjects, cons
 				NewGuid = Spawnable->GetGuid();
 				NewActor = SpawnedActor;
 			}
-			else
-			{
-				UpdateRuntimeInstances();
-
-				OnActorAddedToSequencerEvent.Broadcast(NewActor, PossessableGuid);
-			}
 
 			if (bCreateAndAttachCamera)
 			{
@@ -3129,6 +3150,8 @@ void FSequencer::OnNewActorsDropped(const TArray<UObject*>& DroppedObjects, cons
 					NewCamera->SetActorRotation(FRotator(0.f, -90.f, 0.f));
 				}
 
+				OnActorAddedToSequencerEvent.Broadcast(NewCamera, NewCameraGuid);
+
 				if (bAddSpawnable)
 				{
 					FMovieSceneSpawnable* Spawnable = ConvertToSpawnableInternal(NewCameraGuid);
@@ -3146,8 +3169,6 @@ void FSequencer::OnNewActorsDropped(const TArray<UObject*>& DroppedObjects, cons
 				}
 				else
 				{
-					OnActorAddedToSequencerEvent.Broadcast(NewCamera, NewCameraGuid);
-
 					// Parent it
 					NewCamera->AttachToActor(NewActor, FAttachmentTransformRules::KeepRelativeTransform);
 				}
diff --git a/Engine/Source/Editor/Sequencer/Private/Sequencer.h b/Engine/Source/Editor/Sequencer/Private/Sequencer.h
--- a/Engine/Source/Editor/Sequencer/Private/Sequencer.h
+++ b/Engine/Source/Editor/Sequencer/Private/Sequencer.h
@@ -500,8 +500,8 @@ public:
 	virtual bool IsRecordingLive() const override;
 	virtual float GetCurrentLocalTime(UMovieSceneSequence& InMovieSceneSequence) override;
 	virtual float GetGlobalTime() const override;
-	virtual void SetGlobalTime(float Time, ESnapTimeMode SnapTimeMode = ESnapTimeMode::STM_None) override;
-	virtual void SetGlobalTimeDirectly(float Time, ESnapTimeMode SnapTimeMode = ESnapTimeMode::STM_None) override;
+	virtual void SetGlobalTime(float Time, ESnapTimeMode SnapTimeMode = ESnapTimeMode::STM_None, bool bLooped = false) override;
+	virtual void SetGlobalTimeDirectly(float Time, ESnapTimeMode SnapTimeMode = ESnapTimeMode::STM_None, bool bLooped = false) override;
 	virtual void SetPerspectiveViewportPossessionEnabled(bool bEnabled) override;
 	virtual void SetPerspectiveViewportCameraCutEnabled(bool bEnabled) override;
 	virtual void EnterSilentMode() override { ++SilentModeCount; }
@@ -870,6 +870,9 @@ private:
 		the MovieScene data can change many times per frame.) */
 	bool bNeedTreeRefresh;
 
+	/** When true, the runtime instances need to be updated next frame. */
+	bool bNeedInstanceRefresh;
+
 	/** Stores the playback status to be restored on refresh. */
 	EMovieScenePlayerStatus::Type StoredPlaybackState;
 
diff --git a/Engine/Source/Editor/Sequencer/Private/SequencerNodeTree.cpp b/Engine/Source/Editor/Sequencer/Private/SequencerNodeTree.cpp
--- a/Engine/Source/Editor/Sequencer/Private/SequencerNodeTree.cpp
+++ b/Engine/Source/Editor/Sequencer/Private/SequencerNodeTree.cpp
@@ -468,33 +468,51 @@ const TSharedPtr<FSequencerDisplayNode>& FSequencerNodeTree::GetHoveredNode() co
 	return HoveredNode;
 }
 
+/*
+ * Add node as filtered and include any parent folders
+ */
+static void AddFilteredNode(const TSharedRef<FSequencerDisplayNode>& StartNode, TSet<TSharedRef<const FSequencerDisplayNode>>& OutFilteredNodes)
+{
+	OutFilteredNodes.Add(StartNode);
+
+	// Gather parent folders up the chain
+	TSharedPtr<FSequencerDisplayNode> ParentNode = StartNode->GetParent();
+	while (ParentNode.IsValid() && ParentNode.Get()->GetType() == ESequencerNode::Folder)
+	{
+		OutFilteredNodes.Add(ParentNode.ToSharedRef());
+		ParentNode = ParentNode->GetParent();
+	}
+}
+
 /**
  * Recursively filters nodes
  *
  * @param StartNode			The node to start from
  * @param FilterStrings		The filter strings which need to be matched
  * @param OutFilteredNodes	The list of all filtered nodes
+ * @return Whether the text filter was passed
  */
 static bool FilterNodesRecursive( FSequencer& Sequencer, const TSharedRef<FSequencerDisplayNode>& StartNode, const TArray<FString>& FilterStrings, TSet<TSharedRef<const FSequencerDisplayNode>>& OutFilteredNodes )
 {
-	// assume the filter is acceptable
-	bool PassedTextFilter = true;
-
-	// check each string in the filter strings list against 
+	// check labels - only one of the labels needs to match
+	bool bMatchedLabel = false;
+	bool bObjectHasLabels = false;
 	for (const FString& String : FilterStrings)
 	{
 		if (String.StartsWith(TEXT("label:")))
-		{
+		{				
 			if (!StartNode->GetParent().IsValid() && (String.Len() > 6))
 			{
 				if (StartNode->GetType() == ESequencerNode::Object)
 				{
+					bObjectHasLabels = true;
 					auto ObjectBindingNode = StaticCastSharedRef<FSequencerObjectBindingNode>(StartNode);
 					auto Labels = Sequencer.GetLabelManager().GetObjectLabels(ObjectBindingNode->GetObjectBinding());
 
-					if ((Labels == nullptr) || !Labels->Strings.Contains(String.RightChop(6)))
+					if (Labels != nullptr && Labels->Strings.Contains(String.RightChop(6)))
 					{
-						return false;
+						bMatchedLabel = true;
+						break;
 					}
 				}
 				else
@@ -503,12 +521,20 @@ static bool FilterNodesRecursive( FSequencer& Sequencer, const TSharedRef<FSeque
 				}
 			}
 		}
-		else if (!StartNode->GetDisplayName().ToString().Contains(String)) 
-		{
-			PassedTextFilter = false;
-		}
+	}
+
+	if (bObjectHasLabels && !bMatchedLabel)
+	{
+		return false;
+	}
+
+	// assume the filter is acceptable
+	bool bPassedTextFilter = true;
 
-		if (!PassedTextFilter)
+	// check each string in the filter strings list against 
+	for (const FString& String : FilterStrings)
+	{
+		if (!String.StartsWith(TEXT("label:")) && !StartNode->GetDisplayName().ToString().Contains(String)) 
 		{
 			break;
 		}
@@ -517,10 +543,11 @@ static bool FilterNodesRecursive( FSequencer& Sequencer, const TSharedRef<FSeque
 	// whether or the start node is in the filter
 	bool bInFilter = false;
 
-	if (PassedTextFilter)
+	if (bPassedTextFilter)
 	{
 		// This node is now filtered
-		OutFilteredNodes.Add(StartNode);
+		AddFilteredNode(StartNode, OutFilteredNodes);
+
 		bInFilter = true;
 	}
 
@@ -530,16 +557,17 @@ static bool FilterNodesRecursive( FSequencer& Sequencer, const TSharedRef<FSeque
 	for (const auto& Node : ChildNodes)
 	{
 		// Mark the parent as filtered if any child node was filtered
-		PassedTextFilter |= FilterNodesRecursive(Sequencer, Node, FilterStrings, OutFilteredNodes);
+		bPassedTextFilter |= FilterNodesRecursive(Sequencer, Node, FilterStrings, OutFilteredNodes);
 
-		if (PassedTextFilter && !bInFilter)
+		if (bPassedTextFilter && !bInFilter)
 		{
-			OutFilteredNodes.Add(StartNode);
+			AddFilteredNode(StartNode, OutFilteredNodes);
+
 			bInFilter = true;
 		}
 	}
 
-	return PassedTextFilter;
+	return bPassedTextFilter;
 }
 
 
diff --git a/Engine/Source/Editor/Sequencer/Private/SequencerSettings.cpp b/Engine/Source/Editor/Sequencer/Private/SequencerSettings.cpp
--- a/Engine/Source/Editor/Sequencer/Private/SequencerSettings.cpp
+++ b/Engine/Source/Editor/Sequencer/Private/SequencerSettings.cpp
@@ -24,8 +24,8 @@ USequencerSettings::USequencerSettings( const FObjectInitializer& ObjectInitiali
 	bFixedTimeStepPlayback = true;
 	CurveValueSnapInterval = 10.0f;
 	bSnapCurveValueToInterval = true;
-	bDetailsViewVisible = false;
 	bLabelBrowserVisible = false;
+	ZoomPosition = ESequencerZoomPosition::SZP_CurrentTime;
 	bAutoScrollEnabled = false;
 	bShowCurveEditor = false;
 	bShowCurveEditorCurveToolTips = true;
@@ -306,30 +306,30 @@ void USequencerSettings::SetSnapCurveValueToInterval( bool InbSnapCurveValueToIn
 	}
 }
 
-bool USequencerSettings::GetDetailsViewVisible() const
+bool USequencerSettings::GetLabelBrowserVisible() const
 {
-	return bDetailsViewVisible;
+	return bLabelBrowserVisible;
 }
 
-void USequencerSettings::SetDetailsViewVisible(bool Visible)
+void USequencerSettings::SetLabelBrowserVisible(bool Visible)
 {
-	if (bDetailsViewVisible != Visible)
+	if (bLabelBrowserVisible != Visible)
 	{
-		bDetailsViewVisible = Visible;
+		bLabelBrowserVisible = Visible;
 		SaveConfig();
 	}
 }
 
-bool USequencerSettings::GetLabelBrowserVisible() const
+ESequencerZoomPosition USequencerSettings::GetZoomPosition() const
 {
-	return bLabelBrowserVisible;
+	return ZoomPosition;
 }
 
-void USequencerSettings::SetLabelBrowserVisible(bool Visible)
+void USequencerSettings::SetZoomPosition(ESequencerZoomPosition InZoomPosition)
 {
-	if (bLabelBrowserVisible != Visible)
+	if ( ZoomPosition != InZoomPosition)
 	{
-		bLabelBrowserVisible = Visible;
+		ZoomPosition = InZoomPosition;
 		SaveConfig();
 	}
 }
diff --git a/Engine/Source/Editor/Sequencer/Private/SequencerSettings.h b/Engine/Source/Editor/Sequencer/Private/SequencerSettings.h
--- a/Engine/Source/Editor/Sequencer/Private/SequencerSettings.h
+++ b/Engine/Source/Editor/Sequencer/Private/SequencerSettings.h
@@ -4,7 +4,6 @@
 
 #include "SequencerSettings.generated.h"
 
-
 UENUM()
 enum ESequencerSpawnPosition
 {
@@ -15,6 +14,15 @@ enum ESequencerSpawnPosition
 	SSP_PlaceInFrontOfCamera UMETA(DisplayName="Place in Front of Camera"),
 };
 
+UENUM()
+enum ESequencerZoomPosition
+{
+	/** Current Time. */
+	SZP_CurrentTime UMETA(DisplayName="Current Time"),
+
+	/** Mouse Position. */
+	SZP_MousePosition UMETA(DisplayName="Mouse Position"),
+};
 
 /** Empty class used to house multiple named USequencerSettings */
 UCLASS()
@@ -154,18 +162,16 @@ public:
 	/** Sets whether or not to snap curve values to the interval. */
 	void SetSnapCurveValueToInterval(bool InbSnapCurveValueToInterval);
 
-	/** Gets whether or not the details view is visible. */
-	bool GetDetailsViewVisible() const;
-
-	/** Sets whether or not the details view is visible. */
-	void SetDetailsViewVisible(bool InbDetailsViewVisible);
-
 	/** Gets whether or not the label browser is visible. */
 	bool GetLabelBrowserVisible() const;
-
 	/** Sets whether or not the label browser is visible. */
 	void SetLabelBrowserVisible(bool Visible);
 
+	/** Get zoom in/out position (mouse position or current time). */
+	ESequencerZoomPosition GetZoomPosition() const;
+	/** Set zoom in/out position (mouse position or current time). */
+	void SetZoomPosition(ESequencerZoomPosition InZoomPosition);
+
 	/** Gets whether or not auto-scroll is enabled. */
 	bool GetAutoScrollEnabled() const;
 	/** Sets whether or not auto-scroll is enabled. */
@@ -234,100 +240,100 @@ public:
 
 protected:
 
-	UPROPERTY( config )
+	UPROPERTY( config, EditAnywhere, Category=Keyframing )
 	TEnumAsByte<EAutoKeyMode> AutoKeyMode;
 
-	UPROPERTY( config )
+	UPROPERTY( config, EditAnywhere, Category=Keyframing )
 	bool bKeyAllEnabled;
 
-	UPROPERTY( config )
+	UPROPERTY( config, EditAnywhere, Category=Keyframing )
 	bool bKeyInterpPropertiesOnly;
 
-	UPROPERTY( config )
+	UPROPERTY( config, EditAnywhere, Category=Keyframing )
 	TEnumAsByte<EMovieSceneKeyInterpolation> KeyInterpolation;
 
-	UPROPERTY( config )
+	UPROPERTY( config, EditAnywhere, Category=General )
 	TEnumAsByte<ESequencerSpawnPosition> SpawnPosition;
 
-	UPROPERTY( config )
+	UPROPERTY( config, EditAnywhere, Category=Timeline )
 	bool bShowFrameNumbers;
 
-	UPROPERTY( config )
+	UPROPERTY( config, EditAnywhere, Category=Timeline )
 	bool bShowRangeSlider;
 
-	UPROPERTY( config )
+	UPROPERTY( config, EditAnywhere, Category=Snapping )
 	bool bIsSnapEnabled;
 
-	UPROPERTY( config )
+	UPROPERTY( config, EditAnywhere, Category=Snapping )
 	float TimeSnapInterval;
 
-	UPROPERTY( config )
+	UPROPERTY( config, EditAnywhere, Category=Snapping )
 	bool bSnapKeyTimesToInterval;
 
-	UPROPERTY( config )
+	UPROPERTY( config, EditAnywhere, Category=Snapping )
 	bool bSnapKeyTimesToKeys;
 
-	UPROPERTY( config )
+	UPROPERTY( config, EditAnywhere, Category=Snapping )
 	bool bSnapSectionTimesToInterval;
 
-	UPROPERTY( config )
+	UPROPERTY( config, EditAnywhere, Category=Snapping )
 	bool bSnapSectionTimesToSections;
 
-	UPROPERTY( config )
+	UPROPERTY( config, EditAnywhere, Category=Snapping )
 	bool bSnapPlayTimeToKeys;
 
-	UPROPERTY( config )
+	UPROPERTY( config, EditAnywhere, Category=Snapping )
 	bool bSnapPlayTimeToInterval;
 
-	UPROPERTY( config )
+	UPROPERTY( config, EditAnywhere, Category=Snapping )
 	bool bSnapPlayTimeToDraggedKey;
 
-	UPROPERTY( config )
+	UPROPERTY( config, EditAnywhere, Category=Snapping )
 	bool bFixedTimeStepPlayback;
 
-	UPROPERTY( config )
+	UPROPERTY( config, EditAnywhere, Category=Snapping )
 	float CurveValueSnapInterval;
 
-	UPROPERTY( config )
+	UPROPERTY( config, EditAnywhere, Category=Snapping )
 	bool bSnapCurveValueToInterval;
 
-	UPROPERTY( config )
-	bool bDetailsViewVisible;
-
-	UPROPERTY( config )
+	UPROPERTY( config, EditAnywhere, Category=General )
 	bool bLabelBrowserVisible;
 
-	UPROPERTY( config )
+	UPROPERTY( config, EditAnywhere, Category=Timeline )
+	TEnumAsByte<ESequencerZoomPosition> ZoomPosition;
+
+	UPROPERTY( config, EditAnywhere, Category=Timeline )
 	bool bAutoScrollEnabled;
 
-	UPROPERTY( config )
+	UPROPERTY( config, EditAnywhere, Category=CurveEditor )
 	bool bShowCurveEditor;
 
-	UPROPERTY( config )
+	UPROPERTY( config, EditAnywhere, Category=CurveEditor )
 	bool bShowCurveEditorCurveToolTips;
 
-	UPROPERTY( config )
+	UPROPERTY( config, EditAnywhere, Category=CurveEditor )
 	bool bLinkCurveEditorTimeRange;
 
 	UPROPERTY( config )
 	bool bLooping;
 
-	UPROPERTY( config )
+	UPROPERTY( config, EditAnywhere, Category=Timeline )
 	bool bKeepCursorInPlayRange;
 
-	UPROPERTY( config )
+	UPROPERTY( config, EditAnywhere, Category=Timeline )
 	bool bKeepPlayRangeInSectionBounds;
 
-	UPROPERTY( config )
+	UPROPERTY( config, EditAnywhere, Category=Timeline )
 	uint8 ZeroPadFrames;
 
-	UPROPERTY( config )
+	UPROPERTY( config, EditAnywhere, Category=Timeline )
 	bool bShowCombinedKeyframes;
 
-	UPROPERTY( config )
+	UPROPERTY( config, EditAnywhere, Category=Timeline )
 	bool bInfiniteKeyAreas;
 
-	UPROPERTY( config )
+	UPROPERTY( config, EditAnywhere, Category=Timeline )
 	bool bShowChannelColors;
 
 	UPROPERTY( config )
