Fixed DOF when enabling Pre-Exposure
Added units to the intensity property of point, spot, rect, directional and sky lights.
Refactored the VisualizeHDR so that it shows the HDR view is a PIP on top of the main render.
Fixed PixelInspector results being affected by pre-exposure.
Added a pre-exposure output value to the PixelInspector.
Moved ExposureCompensation to the Exposure tab in post-process settings.
Changed max slider intensity for directional light to be in 120 000 lux
Changed max slider intensity for sky light to be in 50 000 cd/m2
Non bUseInverseSquaredFalloff point lights now have their units forced to ELightUnits::Unitless
Moved light units to advanced display, now that the unit is embeded in the Intensity property
Added Candela, Lux and CandelaPerMeterSquared units to the engine units
Fixed meta data properties to support the "Units" keyword.

diff --git a/Engine/Shaders/Private/MiniFontCommon.ush b/Engine/Shaders/Private/MiniFontCommon.ush
--- a/Engine/Shaders/Private/MiniFontCommon.ush
+++ b/Engine/Shaders/Private/MiniFontCommon.ush
@@ -121,4 +121,31 @@ void PrintFloatNoFraction(int2 PixelPos, inout float3 OutColor, float3 FontColor
 	PrintCharacter(PixelPos, OutColor, FontColor, Cursor, ExtractDigitFromFloat(Number, 1));
 }
 
+// for printf debugging in the shader, has to be positive (other abs is taken)
+// outputs a float number in the form: xx.yy
+// @param LeftTop - in pixels
+void PrintSmallFloat(int2 PixelPos, inout float3 OutColor, float3 FontColor, int2 LeftTop, float Number)
+{
+	Number = abs(Number) + 0.05; // Round up first digit
+	int2 Cursor = LeftTop;
+
+	float BorderDistance = ComputeDistanceToRect(PixelPos, LeftTop, int2(7, 1) * 8 - 1);
+
+	// black border around number
+	//	OutColor = lerp(0, OutColor, saturate(BorderDistance - 2));
+
+	// before period
+	FLATTEN
+	if (Number >= 10)
+	{
+		PrintCharacter(PixelPos, OutColor, FontColor, Cursor, ExtractDigitFromFloat(Number, 10));
+	}
+	PrintCharacter(PixelPos, OutColor, FontColor, Cursor, ExtractDigitFromFloat(Number, 1));
+	// period
+	PrintCharacter(PixelPos, OutColor, FontColor, Cursor, 512 / 8 - 3);
+	// after period
+	PrintCharacter(PixelPos, OutColor, FontColor, Cursor, ExtractDigitFromFloat(Number, 0.1));
+}
+
+
 #endif // __TINY_FONT_COMMON__
diff --git a/Engine/Shaders/Private/PostProcessHistogram.usf b/Engine/Shaders/Private/PostProcessHistogram.usf
--- a/Engine/Shaders/Private/PostProcessHistogram.usf
+++ b/Engine/Shaders/Private/PostProcessHistogram.usf
@@ -56,6 +56,9 @@ void MainCS(
 				if(TexelPos.x < GatherExtentInt.x && TexelPos.y < GatherExtentInt.y)
 				{
 					float4 SceneColor = PostprocessInput0.Load(int3(TexelPos, 0));
+	#if USE_PREEXPOSURE
+					SceneColor.xyz *= View.OneOverPreExposure;
+	#endif
 					float LuminanceVal = max(SceneColor.r, max(SceneColor.g, SceneColor.b));
 					float LogLuminance = ComputeHistogramPositionFromLuminance(LuminanceVal);
 			
diff --git a/Engine/Shaders/Private/PostProcessHistogramCommon.ush b/Engine/Shaders/Private/PostProcessHistogramCommon.ush
--- a/Engine/Shaders/Private/PostProcessHistogramCommon.ush
+++ b/Engine/Shaders/Private/PostProcessHistogramCommon.ush
@@ -127,14 +127,8 @@ float ComputeAverageLuminaneWithoutOutlier(Texture2D HistogramTexture, float Min
 // requires EyeAdaptationParams to be set
 float ComputeEyeAdaptationExposure(Texture2D HistogramTexture)
 {
-#if USE_PREEXPOSURE
-	const float OneOverPreExposure = View.OneOverPreExposure;
-#else
-	const float OneOverPreExposure = 1.f;
-#endif
-
 	float HistogramSum = ComputeHistogramSum(HistogramTexture);
-	float UnclampedAdaptedLuminance = OneOverPreExposure * ComputeAverageLuminaneWithoutOutlier(HistogramTexture, HistogramSum * EyeAdaptationParams[0].x, HistogramSum * EyeAdaptationParams[0].y);
+	float UnclampedAdaptedLuminance = ComputeAverageLuminaneWithoutOutlier(HistogramTexture, HistogramSum * EyeAdaptationParams[0].x, HistogramSum * EyeAdaptationParams[0].y);
 	float ClampedAdaptedLuminance = clamp(UnclampedAdaptedLuminance, EyeAdaptationParams[0].z, EyeAdaptationParams[0].w);
 
 	return ClampedAdaptedLuminance;
diff --git a/Engine/Shaders/Private/PostProcessTemporalAA.usf b/Engine/Shaders/Private/PostProcessTemporalAA.usf
--- a/Engine/Shaders/Private/PostProcessTemporalAA.usf
+++ b/Engine/Shaders/Private/PostProcessTemporalAA.usf
@@ -453,8 +453,8 @@ static const uint kNeighborsCount = 9;
 
 //------------------------------------------------------- PARAMETERS
 
-// 1 / LastFrame.PreExposure
-float OneOverHistoryPreExposure;
+// (Pre-Exposure, 1 / Pre-exposure, LastFramePreExposure, 1 / LastFramePreExposure)
+float4 PreExposureSettings;
 
 float SampleWeights[9];
 float PlusWeights[5];
@@ -785,7 +785,7 @@ float4 TransformBackToRawLinearSceneColor(float4 SceneColor)
 float4 TransformCurrentFrameSceneColor(float4 RawSceneColorRGBA)
 {
 	#if USE_PREEXPOSURE
-		float4 RawLinearSceneColorRGBA = RawSceneColorRGBA * View.OneOverPreExposure;
+		float4 RawLinearSceneColorRGBA = RawSceneColorRGBA *  PreExposureSettings.y;
 	#else
 		float4 RawLinearSceneColorRGBA = RawSceneColorRGBA;
 	#endif
@@ -1832,7 +1832,7 @@ FTAAHistoryPayload SampleHistory(in float2 HistoryScreenPosition)
 	#endif
 
 	#if USE_PREEXPOSURE
-		HistoryPayload.Color.rgb *= OneOverHistoryPreExposure;
+		HistoryPayload.Color.rgb *= PreExposureSettings.w;
 	#endif
 
 	HistoryPayload.Color = TransformSceneColor(HistoryPayload.Color);
@@ -2213,7 +2213,7 @@ FTAAHistoryPayload TemporalAASample(uint2 GroupId, uint2 GroupThreadId, uint Gro
 	OutputPayload.Color = TransformBackToRawLinearSceneColor(OutputPayload.Color);
 	
 	#if USE_PREEXPOSURE
-		OutputPayload.Color.rgb *= View.PreExposure;
+		OutputPayload.Color.rgb *= PreExposureSettings.x;
 	#endif
 
 	#if AA_NAN
diff --git a/Engine/Shaders/Private/PostProcessVisualizeHDR.usf b/Engine/Shaders/Private/PostProcessVisualizeHDR.usf
--- a/Engine/Shaders/Private/PostProcessVisualizeHDR.usf
+++ b/Engine/Shaders/Private/PostProcessVisualizeHDR.usf
@@ -122,15 +122,32 @@ void MainPS(noperspective float4 UVAndScreenPos : TEXCOORD0, float4 SvPosition :
 	int2 PixelPos = (int2)SvPosition.xy;
 	float2 ViewLocalUV = float2(UVAndScreenPos.z * 0.5f + 0.5f, 0.5f - 0.5f * UVAndScreenPos.w);
 
-	int2 ViewportCenter = (int2)(ViewportRect.xy + ViewportSize.xy / 2);
+	// retrieve the exposure scale and target exposure scale from the eye-adaptation buffer.
+	float EyeAdaptationResult = EyeAdaptationTexture.Load(int3(0, 0, 0)).r;
+	float EyeAdaptationTarget = EyeAdaptationTexture.Load(int3(0, 0, 0)).g;
 
+	int2 ViewportCenter = (int2)(ViewportRect.xy + ViewportSize.xy / 2);
 
 	// background is the scene color
-	float4 SceneColor = Texture2DSample(PostprocessInput0, PostprocessInput0Sampler, UV);
+	float2 PIPTopLeft = float2(.6, .2);
+	float2 PIPBottomRight = float2(.9, .5);
+	float2 PIPUV = (ViewLocalUV - PIPTopLeft) / (PIPBottomRight - PIPTopLeft);
 
-//	OutColor = SceneColor;
-	float LuminanceVal = max(SceneColor.r, max(SceneColor.g, SceneColor.b));
-	OutColor = float4(Colorize(ComputeHistogramPositionFromLuminance(LuminanceVal)), 1.0f);
+	if (PIPUV.x == saturate(PIPUV.x) && PIPUV.y == saturate(PIPUV.y))
+	{
+		// Remap from (0, 0) - (1, 1) to ViewportRect / BufferSize
+		PIPUV = PIPUV * View.ViewSizeAndInvSize.xy * View.BufferSizeAndInvSize.zw + ViewportRect.xy * ViewportSize.zw;
+		float4 SceneColor = Texture2DSample(PostprocessInput2, PostprocessInput2Sampler, PIPUV);
+	#if USE_PREEXPOSURE
+		SceneColor.xyz *= View.OneOverPreExposure;
+	#endif
+		float LuminanceVal = max(SceneColor.r, max(SceneColor.g, SceneColor.b));
+		OutColor = float4(Colorize(ComputeHistogramPositionFromLuminance(LuminanceVal)), 1.0f);
+	}
+	else
+	{
+		OutColor = Texture2DSample(PostprocessInput0, PostprocessInput0Sampler, UV);
+	}
 
 	// Compute and apply weight value at this location 
 	float WeightSlope = EyeAdaptationParams[2].w;
@@ -138,42 +155,6 @@ void MainPS(noperspective float4 UVAndScreenPos : TEXCOORD0, float4 SvPosition :
 	float weight = AdaptationWeight(NdPos.x, WeightSlope) * AdaptationWeight(NdPos.y, WeightSlope);
 	OutColor *= max(weight, 0.05f);  
 
-	// crosshair
-	{
-		float CrossHairMask = PixelPos.x == ViewportCenter.x || PixelPos.y == ViewportCenter.y;
-		float2 DistAbs = abs(PixelPos - ViewportCenter);
-		float Dist = max(DistAbs.x, DistAbs.y);
-		float DistMask = Dist >= 2 && Dist < 7;
-
-		OutColor.xyz = lerp(OutColor.xyz, float3(1, 1, 1), CrossHairMask * DistMask);
-	}
-
-	// value under crosshair
-	{
-		float3 CenterViewportHDRColor = Texture2DSample(PostprocessInput0, PostprocessInput0Sampler, ViewportCenterUV).rgb;
-		float Value = Luminance(CenterViewportHDRColor);
-
-		int2 Cursor;
-
-		Cursor = ViewportCenter + int2(-35, 7 + 0 * 8);
-		PrintCharacter(PixelPos, OutColor.xyz, float3(1, 1, 1), Cursor, 10 + 11);	// 'L'
-		PrintFloat(PixelPos, OutColor.xyz, float3(1, 1, 1), Cursor, Value);
-	
-		Cursor = ViewportCenter + int2(-35, 7 + 1 * 8); 
-		PrintCharacter(PixelPos, OutColor.xyz, float3(1, 0, 0), Cursor, 10 + 17);	// 'R'
-		PrintFloat(PixelPos, OutColor.xyz, float3(1, 1, 1), Cursor, CenterViewportHDRColor.r);
-
-		Cursor = ViewportCenter + int2(-35, 7 + 2 * 8); 
-		PrintCharacter(PixelPos, OutColor.xyz, float3(0, 1, 0), Cursor, 10 + 6);	// 'G'
-		PrintFloat(PixelPos, OutColor.xyz, float3(1, 1, 1), Cursor, CenterViewportHDRColor.g);
-
-		Cursor = ViewportCenter + int2(-35, 7 + 3 * 8); 
-		PrintCharacter(PixelPos, OutColor.xyz, float3(0, 0, 1), Cursor, 10 + 1);	// 'B'
-		PrintFloat(PixelPos, OutColor.xyz, float3(1, 1, 1), Cursor, CenterViewportHDRColor.b);
-	}
-
-
-
 	float2 IDAreaLocalUV = ViewLocalUV * 2 + float2(-1, 0);
 
 	// disabled for now, useful for debgugging the histogram content
@@ -216,7 +197,7 @@ void MainPS(noperspective float4 UVAndScreenPos : TEXCOORD0, float4 SvPosition :
 
 
 	// not fully functional yet
-//	HighlightAdvice(OutColor.rgb, UV, PixelPos);
+	//	HighlightAdvice(OutColor.rgb, UV, PixelPos);
 
 	// left top of the border
 	const int2 HistogramLeftTop = int2(ViewportRect.x + 64, ViewportRect.w - 128 - 32);
@@ -280,9 +261,7 @@ void MainPS(noperspective float4 UVAndScreenPos : TEXCOORD0, float4 SvPosition :
 		OutColor.xyz = lerp(OutColor.xyz, Colorize(InsetUV.x), 0.5f);
 	}
 
-	// retrieve the exposure scale and target exposure scale from the eye-adaptation buffer.
-	float EyeAdaptationResult = EyeAdaptationTexture.Load(int3(0, 0, 0)).r;
-	float EyeAdaptationTarget = EyeAdaptationTexture.Load(int3(0, 0, 0)).g;
+	//if (false)
 	{
 		// HDR luminance >0
 		float LuminanceVal = ComputeLuminanceFromHistogramPosition(InsetUV.x);
@@ -305,11 +284,14 @@ void MainPS(noperspective float4 UVAndScreenPos : TEXCOORD0, float4 SvPosition :
 
 	// eye adaptation
 	{
-		float EyeAdaptationValue = ComputeHistogramPositionFromLuminance(1.0f / EyeAdaptationResult);
+		const float OneOverEyeAdaptationResult = 1.0f / EyeAdaptationResult;
+
+		const float EyeAdaptationValue = ComputeHistogramPositionFromLuminance(OneOverEyeAdaptationResult);
+		const float ValuePx = EyeAdaptationValue * HistogramSize.x;
 
-		float ValuePx = EyeAdaptationValue * HistogramSize.x;
+		const float EyeAdaptationEV100 = log2(OneOverEyeAdaptationResult / 1.2f);
 
-		PrintFloat(PixelPos, OutColor.xyz, float3(1, 1, 1), HistogramLeftTop + int2(ValuePx + - 3 * 8 - 3, 1), EyeAdaptationValue);
+		PrintSmallFloat(PixelPos, OutColor.xyz, float3(1, 1, 1), HistogramLeftTop + int2(ValuePx + - 3 * 8 - 3, 1), EyeAdaptationEV100);
 
 		if(abs(InsetPx.x - ValuePx) < 2 && PixelPos.y > HistogramLeftTop.y + 9)
 		{
diff --git a/Engine/Source/Editor/DetailCustomizations/Private/DirectionalLightComponentDetails.cpp b/Engine/Source/Editor/DetailCustomizations/Private/DirectionalLightComponentDetails.cpp
--- a/Engine/Source/Editor/DetailCustomizations/Private/DirectionalLightComponentDetails.cpp
+++ b/Engine/Source/Editor/DetailCustomizations/Private/DirectionalLightComponentDetails.cpp
@@ -37,7 +37,11 @@ void FDirectionalLightComponentDetails::CustomizeDetails( IDetailLayoutBuilder&
 	TSharedPtr<IPropertyHandle> LightIntensityProperty = DetailBuilder.GetProperty("Intensity", ULightComponentBase::StaticClass());
 	// Point lights need to override the ui min and max for units of lumens, so we have to undo that
 	LightIntensityProperty->SetInstanceMetaData("UIMin",TEXT("0.0f"));
-	LightIntensityProperty->SetInstanceMetaData("UIMax",TEXT("20.0f"));
+	LightIntensityProperty->SetInstanceMetaData("UIMax",TEXT("120000.0f"));
+	LightIntensityProperty->SetInstanceMetaData("SliderExponent", TEXT("10.0f"));
+	LightIntensityProperty->SetInstanceMetaData("Units", TEXT("lux"));
+	LightIntensityProperty->SetToolTipText(LOCTEXT("DirectionalLightIntensityToolTipText", "Maximum illumination from the light in lux"));
+
 }
 
 bool FDirectionalLightComponentDetails::IsLightMovable() const
diff --git a/Engine/Source/Editor/DetailCustomizations/Private/LocalLightComponentDetails.cpp b/Engine/Source/Editor/DetailCustomizations/Private/LocalLightComponentDetails.cpp
--- a/Engine/Source/Editor/DetailCustomizations/Private/LocalLightComponentDetails.cpp
+++ b/Engine/Source/Editor/DetailCustomizations/Private/LocalLightComponentDetails.cpp
@@ -6,6 +6,7 @@
 #include "Components/SpotLightComponent.h"
 #include "DetailLayoutBuilder.h"
 #include "DetailCategoryBuilder.h"
+#include "Engine/Scene.h"
 
 #define LOCTEXT_NAMESPACE "LocalLightComponentDetails"
 
@@ -37,6 +38,16 @@ void FLocalLightComponentDetails::CustomizeDetails( IDetailLayoutBuilder& Detail
 	LightIntensityProperty->SetInstanceMetaData("UIMin",TEXT("0.0f"));
 	LightIntensityProperty->SetInstanceMetaData("UIMax",  *FString::SanitizeFloat(100000.0f * ConversionFactor));
 	LightIntensityProperty->SetInstanceMetaData("SliderExponent", TEXT("2.0f"));
+	if (Component->IntensityUnits == ELightUnits::Lumens)
+	{
+		LightIntensityProperty->SetInstanceMetaData("Units", TEXT("lm"));
+		LightIntensityProperty->SetToolTipText(LOCTEXT("LightIntensityInLumensToolTipText", "Luminous power or flux in lumens"));
+	}
+	else if (Component->IntensityUnits == ELightUnits::Candelas)
+	{
+		LightIntensityProperty->SetInstanceMetaData("Units", TEXT("cd"));
+		LightIntensityProperty->SetToolTipText(LOCTEXT("LightIntensityInCandelasToolTipText", "Luminous intensity in candelas"));
+	}
 
 	// Make these come first
 	IDetailCategoryBuilder& LightCategory = DetailBuilder.EditCategory( "Light", FText::GetEmpty(), ECategoryPriority::TypeSpecific );
diff --git a/Engine/Source/Editor/DetailCustomizations/Private/PostProcessSettingsCustomization.cpp b/Engine/Source/Editor/DetailCustomizations/Private/PostProcessSettingsCustomization.cpp
--- a/Engine/Source/Editor/DetailCustomizations/Private/PostProcessSettingsCustomization.cpp
+++ b/Engine/Source/Editor/DetailCustomizations/Private/PostProcessSettingsCustomization.cpp
@@ -66,18 +66,20 @@ void FPostProcessSettingsCustomization::CustomizeChildren( TSharedRef<IPropertyH
 	static const FName LegacyTonemapperName("LegacyTonemapper");
 	static const FName TonemapperCategory("Film");
 	static const FName MobileTonemapperCategory("Mobile Tonemapper");
-
 	const bool bDesktopTonemapperFilm = VarTonemapperFilm->GetValueOnGameThread() == 1;
 	const bool bMobileTonemapperFilm = VarMobileTonemapperFilm->GetValueOnGameThread() == 1;
 	const bool bUsingFilmTonemapper = bDesktopTonemapperFilm || bMobileTonemapperFilm;		// Are any platforms use film tonemapper
 	const bool bUsingLegacyTonemapper = !bDesktopTonemapperFilm || !bMobileTonemapperFilm;	// Are any platforms use legacy/ES2 tonemapper
 
+	static const auto VarDefaultAutoExposureExtendDefaultLuminanceRange = IConsoleManager::Get().FindTConsoleVariableDataInt(TEXT("r.DefaultFeature.AutoExposure.ExtendDefaultLuminanceRange"));
+	const bool bExtendedLuminanceRange = VarDefaultAutoExposureExtendDefaultLuminanceRange->GetValueOnGameThread() == 1;
+	static const FName ExposureCategory("Lens|Exposure");
+
 	if(Result == FPropertyAccess::Success && NumChildren > 0)
 	{
 		for( uint32 ChildIndex = 0; ChildIndex < NumChildren; ++ChildIndex )
 		{
 			TSharedPtr<IPropertyHandle> ChildHandle = StructPropertyHandle->GetChildHandle( ChildIndex );
-
 			
 			if( ChildHandle.IsValid() && ChildHandle->GetProperty() )
 			{
@@ -106,6 +108,26 @@ void FPostProcessSettingsCustomization::CustomizeChildren( TSharedRef<IPropertyH
 						}
 					}
 				}
+				else if (CategoryFName == ExposureCategory && bExtendedLuminanceRange)
+				{
+					if (Property->GetName() == TEXT("AutoExposureMinBrightness"))
+					{
+						Property->SetMetaData(TEXT("DisplayName"), TEXT("Min EV100"));
+					}
+					else if (Property->GetName() == TEXT("AutoExposureMaxBrightness"))
+					{
+						Property->SetMetaData(TEXT("DisplayName"), TEXT("Max EV100"));
+					}
+					else if (Property->GetName() == TEXT("HistogramLogMin"))
+					{
+						Property->SetMetaData(TEXT("DisplayName"), TEXT("Histogram Min EV100"));
+					}
+					else if (Property->GetName() == TEXT("HistogramLogMax"))
+					{
+						Property->SetMetaData(TEXT("DisplayName"), TEXT("Histogram Max EV100"));
+					}
+				}
+
 				
 				FString RawCategoryName = CategoryFName.ToString();
 
diff --git a/Engine/Source/Editor/DetailCustomizations/Private/SkyLightComponentDetails.cpp b/Engine/Source/Editor/DetailCustomizations/Private/SkyLightComponentDetails.cpp
--- a/Engine/Source/Editor/DetailCustomizations/Private/SkyLightComponentDetails.cpp
+++ b/Engine/Source/Editor/DetailCustomizations/Private/SkyLightComponentDetails.cpp
@@ -37,7 +37,9 @@ void FSkyLightComponentDetails::CustomizeDetails( IDetailLayoutBuilder& DetailLa
 	{
 		// Point lights need to override the ui min and max for units of lumens, so we have to undo that
 		LightIntensityProperty->SetInstanceMetaData("UIMin", TEXT("0.0f"));
-		LightIntensityProperty->SetInstanceMetaData("UIMax", TEXT("20.0f"));
+		LightIntensityProperty->SetInstanceMetaData("UIMax", TEXT("50000.0f"));
+		LightIntensityProperty->SetInstanceMetaData("SliderExponent", TEXT("10.0f"));
+		LightIntensityProperty->SetInstanceMetaData("Units", TEXT("CandelaPerMeter2"));
 	}
 
 	const TArray< TWeakObjectPtr<UObject> >& SelectedObjects = DetailLayout.GetSelectedObjects();
diff --git a/Engine/Source/Editor/PixelInspector/Private/PixelInspector.cpp b/Engine/Source/Editor/PixelInspector/Private/PixelInspector.cpp
--- a/Engine/Source/Editor/PixelInspector/Private/PixelInspector.cpp
+++ b/Engine/Source/Editor/PixelInspector/Private/PixelInspector.cpp
@@ -438,7 +438,7 @@ namespace PixelInspector
 		TickSinceLastCreateRequest++;
 	}
 	
-	void SPixelInspector::CreatePixelInspectorRequest(FVector2D InspectViewportUV, int32 viewportUniqueId, FSceneInterface *SceneInterface, bool bInGameViewMode)
+	void SPixelInspector::CreatePixelInspectorRequest(FVector2D InspectViewportUV, int32 viewportUniqueId, FSceneInterface *SceneInterface, bool bInGameViewMode, float InPreExposure)
 	{
 		if (TickSinceLastCreateRequest < MINIMUM_TICK_BETWEEN_CREATE_REQUEST)
 			return;
@@ -469,7 +469,7 @@ namespace PixelInspector
 		if (BufferIndex == -1)
 			return;
 		
-		Requests[BufferIndex].SetRequestData(InspectViewportUV, BufferIndex, viewportUniqueId, GBufferFormat, AllowStaticLighting);
+		Requests[BufferIndex].SetRequestData(InspectViewportUV, BufferIndex, viewportUniqueId, GBufferFormat, AllowStaticLighting, InPreExposure);
 		SceneInterface->AddPixelInspectorRequest(&(Requests[BufferIndex]));
 	}
 
@@ -660,22 +660,27 @@ namespace PixelInspector
 
 	void SPixelInspector::ReadBackRequestData()
 	{
-		for (int RequestIndex = 0; RequestIndex < 2; ++RequestIndex)
+		for (int RequestIndex = 0; RequestIndex < ARRAY_COUNT(Requests); ++RequestIndex)
 		{
-			if (Requests[RequestIndex].RequestComplete == false && Requests[RequestIndex].RenderingCommandSend == true)
+			FPixelInspectorRequest& Request = Requests[RequestIndex];
+			if (Request.RequestComplete == false && Request.RenderingCommandSend == true)
 			{
-				if (Requests[RequestIndex].FrameCountAfterRenderingCommandSend >= WAIT_FRAMENUMBER_BEFOREREADING)
+				if (Request.FrameCountAfterRenderingCommandSend >= WAIT_FRAMENUMBER_BEFOREREADING)
 				{
-					if (Requests[RequestIndex].SourceViewportUV == FVector2D(-1, -1))
+					if (Request.SourceViewportUV == FVector2D(-1, -1))
 					{
 						continue;
 					}
+
+
 					PixelInspectorResult PixelResult;
-					PixelResult.ViewportUV = Requests[RequestIndex].SourceViewportUV;
-					PixelResult.ViewUniqueId = Requests[RequestIndex].ViewId;
+					PixelResult.ViewportUV = Request.SourceViewportUV;
+					PixelResult.ViewUniqueId = Request.ViewId;
+					PixelResult.PreExposure = Request.PreExposure;
+					PixelResult.OneOverPreExposure = Request.PreExposure > 0.f ? (1.f / Request.PreExposure) : 1.f;;
 
 					TArray<FColor> BufferFinalColorValue;
-					FTextureRenderTargetResource* RTResourceFinalColor = Buffer_FinalColor_RGB8[Requests[RequestIndex].BufferIndex]->GameThread_GetRenderTargetResource();
+					FTextureRenderTargetResource* RTResourceFinalColor = Buffer_FinalColor_RGB8[Request.BufferIndex]->GameThread_GetRenderTargetResource();
 					if (RTResourceFinalColor->ReadPixels(BufferFinalColorValue) == false)
 					{
 						BufferFinalColorValue.Empty();
@@ -683,17 +688,17 @@ namespace PixelInspector
 					PixelResult.DecodeFinalColor(BufferFinalColorValue);
 
 					TArray<FLinearColor> BufferSceneColorValue;
-					FTextureRenderTargetResource* RTResourceSceneColor = Buffer_SceneColor_Float[Requests[RequestIndex].BufferIndex]->GameThread_GetRenderTargetResource();
+					FTextureRenderTargetResource* RTResourceSceneColor = Buffer_SceneColor_Float[Request.BufferIndex]->GameThread_GetRenderTargetResource();
 					if (RTResourceSceneColor->ReadLinearColorPixels(BufferSceneColorValue) == false)
 					{
 						BufferSceneColorValue.Empty();
 					}
 					PixelResult.DecodeSceneColor(BufferSceneColorValue);
-					
-					if (Buffer_Depth_Float[Requests[RequestIndex].BufferIndex] != nullptr)
+
+					if (Buffer_Depth_Float[Request.BufferIndex] != nullptr)
 					{
 						TArray<FLinearColor> BufferDepthValue;
-						FTextureRenderTargetResource* RTResourceDepth = Buffer_Depth_Float[Requests[RequestIndex].BufferIndex]->GameThread_GetRenderTargetResource();
+						FTextureRenderTargetResource* RTResourceDepth = Buffer_Depth_Float[Request.BufferIndex]->GameThread_GetRenderTargetResource();
 						if (RTResourceDepth->ReadLinearColorPixels(BufferDepthValue) == false)
 						{
 							BufferDepthValue.Empty();
@@ -702,78 +707,78 @@ namespace PixelInspector
 					}
 
 					TArray<FLinearColor> BufferHDRValue;
-					FTextureRenderTargetResource* RTResourceHDR = Buffer_HDR_Float[Requests[RequestIndex].BufferIndex]->GameThread_GetRenderTargetResource();
+					FTextureRenderTargetResource* RTResourceHDR = Buffer_HDR_Float[Request.BufferIndex]->GameThread_GetRenderTargetResource();
 					if (RTResourceHDR->ReadLinearColorPixels(BufferHDRValue) == false)
 					{
 						BufferHDRValue.Empty();
 					}
 					PixelResult.DecodeHDR(BufferHDRValue);
 
-					if (Requests[RequestIndex].GBufferPrecision == EGBufferFormat::Force8BitsPerChannel)
+					if (Request.GBufferPrecision == EGBufferFormat::Force8BitsPerChannel)
 					{
 						TArray<FColor> BufferAValue;
-						FTextureRenderTargetResource* RTResourceA = Buffer_A_RGB8[Requests[RequestIndex].BufferIndex]->GameThread_GetRenderTargetResource();
+						FTextureRenderTargetResource* RTResourceA = Buffer_A_RGB8[Request.BufferIndex]->GameThread_GetRenderTargetResource();
 						if (RTResourceA->ReadPixels(BufferAValue) == false)
 						{
 							BufferAValue.Empty();
 						}
 
 						TArray<FColor> BufferBCDEValue;
-						FTextureRenderTargetResource* RTResourceBCDE = Buffer_BCDE_RGB8[Requests[RequestIndex].BufferIndex]->GameThread_GetRenderTargetResource();
+						FTextureRenderTargetResource* RTResourceBCDE = Buffer_BCDE_RGB8[Request.BufferIndex]->GameThread_GetRenderTargetResource();
 						if (RTResourceA->ReadPixels(BufferBCDEValue) == false)
 						{
 							BufferBCDEValue.Empty();
 						}
 
-						PixelResult.DecodeBufferData(BufferAValue, BufferBCDEValue, Requests[RequestIndex].AllowStaticLighting);
+						PixelResult.DecodeBufferData(BufferAValue, BufferBCDEValue, Request.AllowStaticLighting);
 					}
-					else if (Requests[RequestIndex].GBufferPrecision == EGBufferFormat::Default)
+					else if (Request.GBufferPrecision == EGBufferFormat::Default)
 					{
 						//PF_A2B10G10R10 format is not support yet
 						TArray<FLinearColor> BufferAValue;
-						FTextureRenderTargetResource* RTResourceA = Buffer_A_RGB10[Requests[RequestIndex].BufferIndex]->GameThread_GetRenderTargetResource();
+						FTextureRenderTargetResource* RTResourceA = Buffer_A_RGB10[Request.BufferIndex]->GameThread_GetRenderTargetResource();
 						if (RTResourceA->ReadLinearColorPixels(BufferAValue) == false)
 						{
 							BufferAValue.Empty();
 						}
 
 						TArray<FColor> BufferBCDEValue;
-						FTextureRenderTargetResource* RTResourceBCDE = Buffer_BCDE_RGB8[Requests[RequestIndex].BufferIndex]->GameThread_GetRenderTargetResource();
+						FTextureRenderTargetResource* RTResourceBCDE = Buffer_BCDE_RGB8[Request.BufferIndex]->GameThread_GetRenderTargetResource();
 						if (RTResourceBCDE->ReadPixels(BufferBCDEValue) == false)
 						{
 							BufferBCDEValue.Empty();
 						}
-						PixelResult.DecodeBufferData(BufferAValue, BufferBCDEValue, Requests[RequestIndex].AllowStaticLighting);
+						PixelResult.DecodeBufferData(BufferAValue, BufferBCDEValue, Request.AllowStaticLighting);
 					}
-					else if (Requests[RequestIndex].GBufferPrecision == EGBufferFormat::HighPrecisionNormals || Requests[RequestIndex].GBufferPrecision == EGBufferFormat::Force16BitsPerChannel)
+					else if (Request.GBufferPrecision == EGBufferFormat::HighPrecisionNormals || Request.GBufferPrecision == EGBufferFormat::Force16BitsPerChannel)
 					{
 						//PF_A2B10G10R10 format is not support yet
 						TArray<FFloat16Color> BufferAValue;
-						FTextureRenderTargetResource* RTResourceA = Buffer_A_Float[Requests[RequestIndex].BufferIndex]->GameThread_GetRenderTargetResource();
+						FTextureRenderTargetResource* RTResourceA = Buffer_A_Float[Request.BufferIndex]->GameThread_GetRenderTargetResource();
 						if (RTResourceA->ReadFloat16Pixels(BufferAValue) == false)
 						{
 							BufferAValue.Empty();
 						}
 
 						TArray<FFloat16Color> BufferBCDEValue;
-						FTextureRenderTargetResource* RTResourceBCDE = Buffer_BCDE_Float[Requests[RequestIndex].BufferIndex]->GameThread_GetRenderTargetResource();
+						FTextureRenderTargetResource* RTResourceBCDE = Buffer_BCDE_Float[Request.BufferIndex]->GameThread_GetRenderTargetResource();
 						if (RTResourceA->ReadFloat16Pixels(BufferBCDEValue) == false)
 						{
 							BufferBCDEValue.Empty();
 						}
-						PixelResult.DecodeBufferData(BufferAValue, BufferBCDEValue, Requests[RequestIndex].AllowStaticLighting);
+						PixelResult.DecodeBufferData(BufferAValue, BufferBCDEValue, Request.AllowStaticLighting);
 					}
 					else
 					{
-						checkf(0, TEXT("Unhandled gbuffer format (%i) during pixel inspector readback."), Requests[RequestIndex].GBufferPrecision);
+						checkf(0, TEXT("Unhandled gbuffer format (%i) during pixel inspector readback."), Request.GBufferPrecision);
 					}
 
 					AccumulationResult.Add(PixelResult);
 					ReleaseBuffers(RequestIndex);
-					Requests[RequestIndex].RequestComplete = true;
-					Requests[RequestIndex].RenderingCommandSend = true;
-					Requests[RequestIndex].FrameCountAfterRenderingCommandSend = 0;
-					Requests[RequestIndex].RequestTickSinceCreation = 0;
+					Request.RequestComplete = true;
+					Request.RenderingCommandSend = true;
+					Request.FrameCountAfterRenderingCommandSend = 0;
+					Request.RequestTickSinceCreation = 0;
 				}
 				else
 				{
diff --git a/Engine/Source/Editor/PixelInspector/Private/PixelInspector.h b/Engine/Source/Editor/PixelInspector/Private/PixelInspector.h
--- a/Engine/Source/Editor/PixelInspector/Private/PixelInspector.h
+++ b/Engine/Source/Editor/PixelInspector/Private/PixelInspector.h
@@ -83,8 +83,10 @@ namespace PixelInspector
 		 * ScreenPosition: This is the viewport coordinate in 2D of the pixel to analyze
 		 * viewportUniqueId: The id of the view (FSceneView::State::GetViewKey) we want to capture the pixel, ScreenPosition has to come from this viewport
 		 * SceneInterface: The interface to set the data for the next render frame.
+		 * bInGameViewMode: Whether the next render frame is a game view.
+		 * InPreExposure: The preexposure for the next render frame.
 		*/
-		 void CreatePixelInspectorRequest(FVector2D InspectViewportUV, int32 viewportUniqueId, FSceneInterface *SceneInterface, bool bInGameViewMode);
+		 void CreatePixelInspectorRequest(FVector2D InspectViewportUV, int32 viewportUniqueId, FSceneInterface *SceneInterface, bool bInGameViewMode, float InPreExposure);
 
 		 /**
 		 * Look if there is some request ready to be read and retrieve the value.
diff --git a/Engine/Source/Editor/PixelInspector/Private/PixelInspectorModule.cpp b/Engine/Source/Editor/PixelInspector/Private/PixelInspectorModule.cpp
--- a/Engine/Source/Editor/PixelInspector/Private/PixelInspectorModule.cpp
+++ b/Engine/Source/Editor/PixelInspector/Private/PixelInspectorModule.cpp
@@ -86,13 +86,13 @@ bool FPixelInspectorModule::GetViewportRealtime(int32 ViewportUid, bool IsCurren
 	return IsCurrentlyRealtime;
 }
 
-void FPixelInspectorModule::CreatePixelInspectorRequest(FVector2D InspectViewportUV, int32 ViewportUniqueId, FSceneInterface *SceneInterface, bool bInGameViewMode)
+void FPixelInspectorModule::CreatePixelInspectorRequest(FVector2D InspectViewportUV, int32 ViewportUniqueId, FSceneInterface *SceneInterface, bool bInGameViewMode, float InPreExposure)
 {
 	if (HPixelInspectorWindow.IsValid() == false)
 	{
 		return;
 	}
-	HPixelInspectorWindow->CreatePixelInspectorRequest(InspectViewportUV, ViewportUniqueId, SceneInterface, bInGameViewMode);
+	HPixelInspectorWindow->CreatePixelInspectorRequest(InspectViewportUV, ViewportUniqueId, SceneInterface, bInGameViewMode, InPreExposure);
 }
 
 void FPixelInspectorModule::SetViewportInformation(int32 ViewportUniqueId, FIntPoint ViewportSize)
diff --git a/Engine/Source/Editor/PixelInspector/Private/PixelInspectorResult.cpp b/Engine/Source/Editor/PixelInspector/Private/PixelInspectorResult.cpp
--- a/Engine/Source/Editor/PixelInspector/Private/PixelInspectorResult.cpp
+++ b/Engine/Source/Editor/PixelInspector/Private/PixelInspectorResult.cpp
@@ -35,7 +35,7 @@ namespace PixelInspector
 			SceneColor = FLinearColor::Black;
 			return;
 		}
-		SceneColor = BufferSceneColorValue[0];
+		SceneColor = BufferSceneColorValue[0] * OneOverPreExposure;
 		//Set the alpha to 1.0 as the default value
 		SceneColor.A = 1.0f;
 	}
@@ -57,8 +57,8 @@ namespace PixelInspector
 			HdrColor = FLinearColor::Black;
 			return;
 		}
-		HdrLuminance = BufferHDRValue[0].GetLuminance();
-		HdrColor = BufferHDRValue[0];
+		HdrLuminance = BufferHDRValue[0].GetLuminance() * OneOverPreExposure;
+		HdrColor = BufferHDRValue[0] * OneOverPreExposure;
 	}
 
 	void PixelInspectorResult::DecodeBufferData(TArray<FColor> &BufferAValue, TArray<FColor> &BufferBCDEValue, bool AllowStaticLighting)
diff --git a/Engine/Source/Editor/PixelInspector/Private/PixelInspectorResult.h b/Engine/Source/Editor/PixelInspector/Private/PixelInspectorResult.h
--- a/Engine/Source/Editor/PixelInspector/Private/PixelInspectorResult.h
+++ b/Engine/Source/Editor/PixelInspector/Private/PixelInspectorResult.h
@@ -27,6 +27,8 @@ namespace PixelInspector
 			ViewUniqueId = -1;
 			ViewportUV = FVector2D(-1, -1);
 
+			PreExposure = 1.f;
+
 			Depth = 0.0f;
 			WorldPosition = FVector(0.0f);
 
@@ -59,6 +61,12 @@ namespace PixelInspector
 		int32 ViewUniqueId;
 		FVector2D ViewportUV;
 
+
+		//////////////////////////////////////////////////////////////////////////
+		// PreExposure used to render this frame. See "r.UsePreExposure"
+		float PreExposure;
+		float OneOverPreExposure;
+
 		//////////////////////////////////////////////////////////////////////////
 		// Final color 3x3 grid
 		TArray<FLinearColor> FinalColor;
diff --git a/Engine/Source/Editor/PixelInspector/Private/PixelInspectorView.cpp b/Engine/Source/Editor/PixelInspector/Private/PixelInspectorView.cpp
--- a/Engine/Source/Editor/PixelInspector/Private/PixelInspectorView.cpp
+++ b/Engine/Source/Editor/PixelInspector/Private/PixelInspectorView.cpp
@@ -12,6 +12,7 @@ UPixelInspectorView::UPixelInspectorView(const FObjectInitializer& ObjectInitial
 	}
 	FinalColor = FLinearColor::Green;
 	SceneColor = FLinearColor::Green;
+	PreExposure = 1.f;
 	Luminance = 0.0f;
 	HdrColor = FLinearColor::Black;
 	Normal = FVector(0.0f);
@@ -60,6 +61,7 @@ void UPixelInspectorView::SetFromResult(PixelInspector::PixelInspectorResult &Re
 	}
 	
 	SceneColor = Result.SceneColor;
+	PreExposure = Result.PreExposure;
 	Luminance = Result.HdrLuminance;
 	HdrColor = Result.HdrColor;
 	Normal = Result.Normal;
diff --git a/Engine/Source/Editor/PixelInspector/Private/PixelInspectorView.h b/Engine/Source/Editor/PixelInspector/Private/PixelInspectorView.h
--- a/Engine/Source/Editor/PixelInspector/Private/PixelInspectorView.h
+++ b/Engine/Source/Editor/PixelInspector/Private/PixelInspectorView.h
@@ -25,6 +25,10 @@ class UPixelInspectorView : public UObject
 	UPROPERTY(VisibleAnywhere, category = SceneColor)
 	FLinearColor SceneColor;
 
+	/** HDR Luminance. */
+	UPROPERTY(VisibleAnywhere, category = SceneColor)
+	float PreExposure;
+
 	/** HDR Luminance. */
 	UPROPERTY(VisibleAnywhere, category = HDR)
 	float Luminance;
diff --git a/Engine/Source/Editor/PixelInspector/Public/PixelInspectorModule.h b/Engine/Source/Editor/PixelInspector/Public/PixelInspectorModule.h
--- a/Engine/Source/Editor/PixelInspector/Public/PixelInspectorModule.h
+++ b/Engine/Source/Editor/PixelInspector/Public/PixelInspectorModule.h
@@ -43,7 +43,7 @@ public:
 
 	virtual bool GetViewportRealtime(int32 ViewportUid, bool IsCurrentlyRealtime, bool IsMouseInsideViewport);
 
-	virtual void CreatePixelInspectorRequest(FVector2D InspectViewportUV, int32 ViewportUniqueId, FSceneInterface *SceneInterface, bool bInGameViewMode);
+	virtual void CreatePixelInspectorRequest(FVector2D InspectViewportUV, int32 ViewportUniqueId, FSceneInterface *SceneInterface, bool bInGameViewMode, float InPreExposure);
 
 	virtual void SetViewportInformation(int32 ViewportUniqueId, FIntPoint ViewportSize);
 
diff --git a/Engine/Source/Editor/PropertyEditor/Private/UserInterface/PropertyEditor/SPropertyEditorNumeric.h b/Engine/Source/Editor/PropertyEditor/Private/UserInterface/PropertyEditor/SPropertyEditorNumeric.h
--- a/Engine/Source/Editor/PropertyEditor/Private/UserInterface/PropertyEditor/SPropertyEditorNumeric.h
+++ b/Engine/Source/Editor/PropertyEditor/Private/UserInterface/PropertyEditor/SPropertyEditorNumeric.h
@@ -300,7 +300,14 @@ public:
 				if (FUnitConversion::Settings().ShouldDisplayUnits())
 				{
 					const FString& DynamicUnits = InPropertyEditor->GetProperty()->GetMetaData(TEXT("Units"));
-					PropertyUnits = FUnitConversion::UnitFromString(*DynamicUnits);
+					if (!DynamicUnits.IsEmpty())
+					{
+						PropertyUnits = FUnitConversion::UnitFromString(*DynamicUnits);
+					}
+					else
+					{
+						PropertyUnits = FUnitConversion::UnitFromString(*GetMetaDataFromKey("Units"));
+					}
 				}
 
 				if (!PropertyUnits.IsSet())
diff --git a/Engine/Source/Editor/UnrealEd/Private/EditorViewportClient.cpp b/Engine/Source/Editor/UnrealEd/Private/EditorViewportClient.cpp
--- a/Engine/Source/Editor/UnrealEd/Private/EditorViewportClient.cpp
+++ b/Engine/Source/Editor/UnrealEd/Private/EditorViewportClient.cpp
@@ -3471,7 +3471,7 @@ void FEditorViewportClient::SetupViewForRendering(FSceneViewFamily& ViewFamily,
 			(InspectViewportPos.X + 0.5f) / float(View.UnscaledViewRect.Width()),
 			(InspectViewportPos.Y + 0.5f) / float(View.UnscaledViewRect.Height()));
 
-		PixelInspectorModule.CreatePixelInspectorRequest(InspectViewportUV, View.State->GetViewKey(), SceneInterface, bInGameViewMode);
+		PixelInspectorModule.CreatePixelInspectorRequest(InspectViewportUV, View.State->GetViewKey(), SceneInterface, bInGameViewMode, View.State->GetPreExposure());
 	}
 	else if (!View.bUsePixelInspector && CurrentMousePos != FIntPoint(-1, -1))
 	{
diff --git a/Engine/Source/Editor/UnrealEd/Private/SEditorViewport.cpp b/Engine/Source/Editor/UnrealEd/Private/SEditorViewport.cpp
--- a/Engine/Source/Editor/UnrealEd/Private/SEditorViewport.cpp
+++ b/Engine/Source/Editor/UnrealEd/Private/SEditorViewport.cpp
@@ -559,7 +559,7 @@ TSharedRef<SWidget> SEditorViewport::BuildFixedEV100Menu()  const
 				.MaxValue(EV100Max)
 				.Value( this, &SEditorViewport::OnGetFixedEV100Value )
 				.OnValueChanged( this, &SEditorViewport::OnFixedEV100ValueChanged )
-				.ToolTipText(LOCTEXT( "EV100ToolTip", "Sets the exposure value of the camera using the specified EV100."))
+				.ToolTipText(LOCTEXT( "EV100ToolTip", "Sets the exposure value of the camera using the specified EV100. Exposure = 1 / (1.2 * 2^EV100)"))
 				.IsEnabled( this, &SEditorViewport::IsFixedEV100Enabled )
 			]
 		];
diff --git a/Engine/Source/Runtime/Core/Private/Math/UnitConversion.cpp b/Engine/Source/Runtime/Core/Private/Math/UnitConversion.cpp
--- a/Engine/Source/Runtime/Core/Private/Math/UnitConversion.cpp
+++ b/Engine/Source/Runtime/Core/Private/Math/UnitConversion.cpp
@@ -65,6 +65,9 @@ FParseCandidate ParseCandidates[] = {
 	{ TEXT("Terabytes"),			EUnit::Terabytes },				{ TEXT("TB"),		EUnit::Terabytes },
 
 	{ TEXT("Lumens"),				EUnit::Lumens },				{ TEXT("lm"),		EUnit::Lumens },
+	{ TEXT("Candela"),				EUnit::Candela },				{ TEXT("cd"),		EUnit::Candela },
+	{ TEXT("Lux"),					EUnit::Lux },					{ TEXT("lx"),		EUnit::Lux },
+	{ TEXT("CandelaPerMeterSquared"), EUnit::CandelaPerMeter2 },	{ TEXT("CandelaPerMeter2"), EUnit::CandelaPerMeter2 },
 
 	{ TEXT("Milliseconds"),			EUnit::Milliseconds },			{ TEXT("ms"),		EUnit::Milliseconds },
 	{ TEXT("Seconds"),				EUnit::Seconds },				{ TEXT("s"),		EUnit::Seconds },
@@ -102,7 +105,7 @@ const TCHAR* const DisplayStrings[] = {
 
 	TEXT("B"), TEXT("KiB"), TEXT("MiB"), TEXT("GiB"), TEXT("TiB"),
 
-	TEXT("lm"),
+	TEXT("lm"), TEXT("cd"), TEXT("lux"), TEXT("cd/m2"),
 
 	TEXT("ms"), TEXT("s"), TEXT("min"), TEXT("hr"), TEXT("dy"), TEXT("mth"), TEXT("yr"),
 
@@ -133,7 +136,7 @@ const EUnitType UnitTypes[] = {
 
 	EUnitType::DataSize,	EUnitType::DataSize,	EUnitType::DataSize,	EUnitType::DataSize,	EUnitType::DataSize,
 
-	EUnitType::LuminousFlux,
+	EUnitType::LuminousFlux, EUnitType::LuminousIntensity, EUnitType::Illuminance, EUnitType::Luminance,
 
 	EUnitType::Time,		EUnitType::Time,		EUnitType::Time,		EUnitType::Time,		EUnitType::Time,		EUnitType::Time,		EUnitType::Time,
 
diff --git a/Engine/Source/Runtime/Core/Public/Math/UnitConversion.h b/Engine/Source/Runtime/Core/Public/Math/UnitConversion.h
--- a/Engine/Source/Runtime/Core/Public/Math/UnitConversion.h
+++ b/Engine/Source/Runtime/Core/Public/Math/UnitConversion.h
@@ -43,8 +43,8 @@ enum class EUnit : uint8
 	/** Data Size units */
 	Bytes, Kilobytes, Megabytes, Gigabytes, Terabytes,
 
-	/** Luminous flux units */
-	Lumens,
+	/** Luminous flux units, luminous intensity, illuminance, luminance */
+	Lumens, Candela, Lux, CandelaPerMeter2,
 
 	/** Time units */
 	Milliseconds, Seconds, Minutes, Hours, Days, Months, Years,
@@ -62,7 +62,7 @@ enum class EUnit : uint8
 /** Enumeration that specifies particular classes of unit */
 enum class EUnitType
 {
-	Distance, Angle, Speed, Temperature, Mass, Force, Frequency, DataSize, LuminousFlux, Time, PixelDensity, Multipliers, Arbitrary,
+	Distance, Angle, Speed, Temperature, Mass, Force, Frequency, DataSize, LuminousFlux, LuminousIntensity, Illuminance, Luminance, Time, PixelDensity, Multipliers, Arbitrary,
 
 	// Symbolic entry - do not use directly
 	NumberOf,
diff --git a/Engine/Source/Runtime/Engine/Classes/Components/LocalLightComponent.h b/Engine/Source/Runtime/Engine/Classes/Components/LocalLightComponent.h
--- a/Engine/Source/Runtime/Engine/Classes/Components/LocalLightComponent.h
+++ b/Engine/Source/Runtime/Engine/Classes/Components/LocalLightComponent.h
@@ -23,7 +23,7 @@ class ENGINE_API ULocalLightComponent : public ULightComponent
 	 * The peak luminous intensity is measured in candelas,
 	 * while the luminous power is measured in lumens.
 	 */
-	UPROPERTY(EditAnywhere, BlueprintReadOnly, Category=Light, meta=(DisplayName="Intensity Units", EditCondition="bUseInverseSquaredFalloff"))
+	UPROPERTY(EditAnywhere, BlueprintReadOnly, Category=Light, AdvancedDisplay, meta=(DisplayName="Intensity Units", EditCondition="bUseInverseSquaredFalloff"))
 	ELightUnits IntensityUnits;
 
 	UPROPERTY()
diff --git a/Engine/Source/Runtime/Engine/Classes/Engine/RendererSettings.h b/Engine/Source/Runtime/Engine/Classes/Engine/RendererSettings.h
--- a/Engine/Source/Runtime/Engine/Classes/Engine/RendererSettings.h
+++ b/Engine/Source/Runtime/Engine/Classes/Engine/RendererSettings.h
@@ -351,12 +351,6 @@ class ENGINE_API URendererSettings : public UDeveloperSettings
 		ToolTip = "Whether to use the new DOF implementation for Circle DOF method."))
 	uint32 bUseNewAlgorithm : 1;
 
-	UPROPERTY(config, EditAnywhere, Category = Postprocessing, meta = (
-		ConsoleVariable = "r.UsePreExposure", DisplayName = "Apply Pre-exposure before writing to the scene color",
-		ToolTip = "Whether to use pre-exposure to remap the range of the scene color around the camera exposure. This limits the render target range required to support HDR lighting value.",
-		ConfigRestartRequired=true))
-	uint32 bUsePreExposure : 1;
-
 	UPROPERTY(config, EditAnywhere, Category = DefaultSettings, meta = (
 		ConsoleVariable = "r.DefaultFeature.Bloom", DisplayName = "Bloom",
 		ToolTip = "Whether the default for Bloom is enabled or not (postprocess volume/camera/game setting can still override and enable or disable it independently)"))
@@ -382,6 +376,18 @@ class ENGINE_API URendererSettings : public UDeveloperSettings
 		ToolTip = "The default method for AutoExposure(postprocess volume/camera/game setting can still override and enable or disable it independently)"))
 	TEnumAsByte<EAutoExposureMethodUI::Type> DefaultFeatureAutoExposure; 
 
+	UPROPERTY(config, EditAnywhere, Category = DefaultSettings, meta = (
+		ConsoleVariable = "r.DefaultFeature.AutoExposure.ExtendDefaultLuminanceRange", DisplayName = "Extend default luminance range in Auto Exposure settings",
+		ToolTip = "Whether the default values for AutoExposure should support an extended range of scene luminance. Also changes the exposure settings to be expressed in EV100.",
+		ConfigRestartRequired=true))
+	uint32 bExtendDefaultLuminanceRangeInAutoExposureSettings: 1;
+
+	UPROPERTY(config, EditAnywhere, Category = DefaultSettings, meta = (
+		ConsoleVariable = "r.UsePreExposure", DisplayName = "Apply Pre-exposure before writing to the scene color",
+		ToolTip = "Whether to use pre-exposure to remap the range of the scene color around the camera exposure. This limits the render target range required to support HDR lighting value.",
+		ConfigRestartRequired=true))
+	uint32 bUsePreExposure : 1;
+
 	UPROPERTY(config, EditAnywhere, Category = DefaultSettings, meta = (
 		ConsoleVariable = "r.DefaultFeature.MotionBlur", DisplayName = "Motion Blur",
 		ToolTip = "Whether the default for MotionBlur is enabled or not (postprocess volume/camera/game setting can still override and enable or disable it independently)"))
diff --git a/Engine/Source/Runtime/Engine/Classes/Engine/Scene.h b/Engine/Source/Runtime/Engine/Classes/Engine/Scene.h
--- a/Engine/Source/Runtime/Engine/Classes/Engine/Scene.h
+++ b/Engine/Source/Runtime/Engine/Classes/Engine/Scene.h
@@ -431,8 +431,7 @@ USTRUCT(BlueprintType)
 struct FCameraExposureSettings
 {
 	GENERATED_USTRUCT_BODY()
-
-
+		
 	/** Luminance computation method */
 	UPROPERTY(EditAnywhere, BlueprintReadWrite, Category="Exposure", meta=(DisplayName = "Method"))
     TEnumAsByte<enum EAutoExposureMethod> Method;
@@ -504,24 +503,7 @@ struct FCameraExposureSettings
 	UPROPERTY(Interp, BlueprintReadWrite, Category = "Exposure", AdvancedDisplay, meta=(UIMin = "0", UIMax = "100.0", DisplayName = "Calibration Constant"))
 	float CalibrationConstant;
 
-
-	FCameraExposureSettings()
-	{
-		// next value might get overwritten by r.DefaultFeature.AutoExposure.Method
-		Method = AEM_Histogram;
-		LowPercent = 80.0f;
-		HighPercent = 98.3f;
-		// next value might get overwritten by r.DefaultFeature.AutoExposure
-		MinBrightness = 0.03f;
-		// next value might get overwritten by r.DefaultFeature.AutoExposure
-		MaxBrightness = 2.0f;
-		SpeedUp = 3.0f;
-		SpeedDown = 1.0f;
-		Bias = 0.0f;
-		HistogramLogMin = -8.0f;
-		HistogramLogMax = 4.0f;
-		CalibrationConstant	= 16.0;
-	}
+	ENGINE_API FCameraExposureSettings();
 
 	/* Exports to post process settings with overrides. */
 	ENGINE_API void ExportToPostProcessSettings(FPostProcessSettings* OutPostProcessSettings) const;
@@ -1335,7 +1317,7 @@ struct FPostProcessSettings
 	 * Logarithmic adjustment for the exposure. Only used if a tonemapper is specified.
 	 * 0: no adjustment, -1:2x darker, -2:4x darker, 1:2x brighter, 2:4x brighter, ...
 	 */
-	UPROPERTY(interp, BlueprintReadWrite, Category = "Lens|Camera", meta = (UIMin = "-8.0", UIMax = "8.0", editcondition = "bOverride_AutoExposureBias", DisplayName = "Exposure Compensation "))
+	UPROPERTY(interp, BlueprintReadWrite, Category = "Lens|Exposure", meta = (UIMin = "-8.0", UIMax = "8.0", editcondition = "bOverride_AutoExposureBias", DisplayName = "Exposure Compensation "))
 	float AutoExposureBias;
 
 	/**
@@ -1358,22 +1340,19 @@ struct FPostProcessSettings
 	UPROPERTY(interp, BlueprintReadWrite, Category="Lens|Exposure", AdvancedDisplay, meta=(ClampMin = "0.0", ClampMax = "100.0", editcondition = "bOverride_AutoExposureHighPercent", DisplayName = "High Percent"))
 	float AutoExposureHighPercent;
 
+
 	/**
-	 * A good value should be positive near 0. This is the minimum brightness the auto exposure can adapt to.
-	 * It should be tweaked in a dark lighting situation (too small: image appears too bright, too large: image appears too dark).
-	 * Note: Tweaking emissive materials and lights or tweaking auto exposure can look the same. Tweaking auto exposure has global
-	 * effect and defined the HDR range - you don't want to change that late in the project development.
-	 * Eye Adaptation is disabled if MinBrightness = MaxBrightness
+	 * Auto-Exposure minimum adaptation. Eye Adaptation is disabled if Min = Max. 
+	 * Auto-exposure is implemented by choosing an exposure value for which the average luminance generates a pixel brightness equal to the Constant Calibration value.
+	 * The Min/Max are expressed in pixel luminance (cd/m2) or in EV100 when using ExtendDefaultLuminanceRange (see project settings).
 	 */
 	UPROPERTY(interp, BlueprintReadWrite, Category="Lens|Exposure", meta=(ClampMin = "-10.0", UIMax = "20.0", editcondition = "bOverride_AutoExposureMinBrightness", DisplayName = "Min Brightness"))
 	float AutoExposureMinBrightness;
 
 	/**
-	 * A good value should be positive (2 is a good value). This is the maximum brightness the auto exposure can adapt to.
-	 * It should be tweaked in a bright lighting situation (too small: image appears too bright, too large: image appears too dark).
-	 * Note: Tweaking emissive materials and lights or tweaking auto exposure can look the same. Tweaking auto exposure has global
-	 * effect and defined the HDR range - you don't want to change that late in the project development.
-	 * Eye Adaptation is disabled if MinBrightness = MaxBrightness
+	 * Auto-Exposure maximum adaptation. Eye Adaptation is disabled if Min = Max. 
+	 * Auto-exposure is implemented by choosing an exposure value for which the average luminance generates a pixel brightness equal to the Constant Calibration value.
+	 * The Min/Max are expressed in pixel luminance (cd/m2) or in EV100 when using ExtendDefaultLuminanceRange (see project settings).
 	 */
 	UPROPERTY(interp, BlueprintReadWrite, Category="Lens|Exposure", meta=(ClampMin = "-10.0", UIMax = "20.0", editcondition = "bOverride_AutoExposureMaxBrightness", DisplayName = "Max Brightness"))
 	float AutoExposureMaxBrightness;
@@ -1386,11 +1365,11 @@ struct FPostProcessSettings
 	UPROPERTY(interp, BlueprintReadWrite, Category="Lens|Exposure", meta=(ClampMin = "0.02", UIMax = "20.0", editcondition = "bOverride_AutoExposureSpeedDown", DisplayName = "Speed Down"))
 	float AutoExposureSpeedDown;
 
-	/** temporary exposed until we found good values, -8: 1/256, -10: 1/1024 */
+	/** Histogram Min value. Expressed in Log2(Luminance) or in EV100 when using ExtendDefaultLuminanceRange (see project settings) */
 	UPROPERTY(interp, BlueprintReadWrite, Category="Lens|Exposure", AdvancedDisplay, meta=(UIMin = "-16", UIMax = "0.0", editcondition = "bOverride_HistogramLogMin"))
 	float HistogramLogMin;
 
-	/** temporary exposed until we found good values 4: 16, 8: 256 */
+	/** Histogram Max value. Expressed in Log2(Luminance) or in EV100 when using ExtendDefaultLuminanceRange (see project settings) */
 	UPROPERTY(interp, BlueprintReadWrite, Category="Lens|Exposure", AdvancedDisplay, meta=(UIMin = "0.0", UIMax = "16.0", editcondition = "bOverride_HistogramLogMax"))
 	float HistogramLogMax;
 
diff --git a/Engine/Source/Runtime/Engine/Private/Components/PointLightComponent.cpp b/Engine/Source/Runtime/Engine/Private/Components/PointLightComponent.cpp
--- a/Engine/Source/Runtime/Engine/Private/Components/PointLightComponent.cpp
+++ b/Engine/Source/Runtime/Engine/Private/Components/PointLightComponent.cpp
@@ -230,6 +230,11 @@ void UPointLightComponent::Serialize(FArchive& Ar)
 	{
 		AddLocalRotation( FRotator(-90.f, 0.f, 0.f) );
 	}
+
+	if (Ar.IsLoading() && !bUseInverseSquaredFalloff)
+	{
+		IntensityUnits = ELightUnits::Unitless;
+	}
 }
 
 #if WITH_EDITOR
@@ -239,11 +244,14 @@ bool UPointLightComponent::CanEditChange(const UProperty* InProperty) const
 	if (InProperty)
 	{
 		FString PropertyName = InProperty->GetName();
-
-		if (FCString::Strcmp(*PropertyName, TEXT("LightFalloffExponent")) == 0)
+		if (PropertyName == GET_MEMBER_NAME_STRING_CHECKED(UPointLightComponent, LightFalloffExponent))
 		{
 			return !bUseInverseSquaredFalloff;
 		}
+		if (PropertyName == GET_MEMBER_NAME_STRING_CHECKED(ULocalLightComponent, IntensityUnits))
+		{
+			return bUseInverseSquaredFalloff;
+		}
 	}
 
 	return Super::CanEditChange(InProperty);
@@ -262,6 +270,11 @@ void UPointLightComponent::PostEditChangeProperty(FPropertyChangedEvent& Propert
 	SoftSourceRadius = FMath::Max(0.0f, SoftSourceRadius);
 	SourceLength = FMath::Max(0.0f, SourceLength);
 
+	if (!bUseInverseSquaredFalloff)
+	{
+		IntensityUnits = ELightUnits::Unitless;
+	}
+
 	Super::PostEditChangeProperty(PropertyChangedEvent);
 }
 #endif // WITH_EDITOR
diff --git a/Engine/Source/Runtime/Engine/Private/Scene.cpp b/Engine/Source/Runtime/Engine/Private/Scene.cpp
--- a/Engine/Source/Runtime/Engine/Private/Scene.cpp
+++ b/Engine/Source/Runtime/Engine/Private/Scene.cpp
@@ -1,7 +1,7 @@
 // Copyright 1998-2018 Epic Games, Inc. All Rights Reserved.
 
 #include "Engine/Scene.h"
-
+#include "HAL/IConsoleManager.h"
 
 void FColorGradingSettings::ExportToPostProcessSettings(FPostProcessSettings* OutPostProcessSettings) const
 {
@@ -158,6 +158,38 @@ void FLensSettings::ExportToPostProcessSettings(FPostProcessSettings* OutPostPro
 	OutPostProcessSettings->SceneFringeIntensity = ChromaticAberration;
 }
 
+FCameraExposureSettings::FCameraExposureSettings()
+{
+	static const auto VarDefaultAutoExposureExtendDefaultLuminanceRange = IConsoleManager::Get().FindTConsoleVariableDataInt(TEXT("r.DefaultFeature.AutoExposure.ExtendDefaultLuminanceRange"));
+	const bool bExtendedLuminanceRange = VarDefaultAutoExposureExtendDefaultLuminanceRange->GetValueOnAnyThread() == 1;
+
+	// next value might get overwritten by r.DefaultFeature.AutoExposure.Method
+	Method = AEM_Histogram;
+	LowPercent = 80.0f;
+	HighPercent = 98.3f;
+
+	if (bExtendedLuminanceRange)
+	{
+		// When this project setting is set, the following values are in EV100.
+		MinBrightness = -10.0f;
+		MaxBrightness = 20.0f;
+		HistogramLogMin = -10.0f;
+		HistogramLogMax = 20.0f;
+	}
+	else
+	{
+		MinBrightness = 0.03f;
+		MaxBrightness = 2.0f;
+		HistogramLogMin = -8.0f;
+		HistogramLogMax = 4.0f;
+	}
+
+	SpeedUp = 3.0f;
+	SpeedDown = 1.0f;
+	Bias = 0.0f;
+	CalibrationConstant	= 16.0;
+}
+
 void FCameraExposureSettings::ExportToPostProcessSettings(FPostProcessSettings* OutPostProcessSettings) const
 {
 	OutPostProcessSettings->bOverride_AutoExposureMethod = true;
@@ -311,7 +343,6 @@ static void DoPostProcessSettingsSanityCheck()
 
 #endif // DO_CHECK
 
-
 FPostProcessSettings::FPostProcessSettings()
 {
 	// to set all bOverride_.. by default to false
@@ -419,10 +450,25 @@ FPostProcessSettings::FPostProcessSettings()
 	AutoExposureMethod = AEM_Histogram;
 	AutoExposureLowPercent = 80.0f;
 	AutoExposureHighPercent = 98.3f;
+
 	// next value might get overwritten by r.DefaultFeature.AutoExposure
-	AutoExposureMinBrightness = 0.03f;
-	// next value might get overwritten by r.DefaultFeature.AutoExposure
-	AutoExposureMaxBrightness = 2.0f;
+	static const auto VarDefaultAutoExposureExtendDefaultLuminanceRange = IConsoleManager::Get().FindTConsoleVariableDataInt(TEXT("r.DefaultFeature.AutoExposure.ExtendDefaultLuminanceRange"));
+	if (VarDefaultAutoExposureExtendDefaultLuminanceRange->GetValueOnAnyThread() != 0)
+	{
+		// When this project setting is set, the following values are in EV100.
+		AutoExposureMinBrightness = -10.0f;
+		AutoExposureMaxBrightness = 20.0f;
+		HistogramLogMin = -10.0f;
+		HistogramLogMax = 20.0f;
+	}
+	else
+	{
+		AutoExposureMinBrightness = 0.03f;
+		AutoExposureMaxBrightness = 2.0f;
+		HistogramLogMin = -8.0f;
+		HistogramLogMax = 4.0f;
+	}
+
 	AutoExposureBias = 0.0f;
 	AutoExposureSpeedUp = 3.0f;
 	AutoExposureSpeedDown = 1.0f;
@@ -434,8 +480,7 @@ FPostProcessSettings::FPostProcessSettings()
 	LPVSpecularOcclusionIntensity = 1.0f;
 	LPVFadeRange = 0.0f;
 	LPVDirectionalOcclusionFadeRange = 0.0f;
-	HistogramLogMin = -8.0f;
-	HistogramLogMax = 4.0f;
+
 	// next value might get overwritten by r.DefaultFeature.LensFlare
 	LensFlareIntensity = 1.0f;
 	LensFlareTint = FLinearColor(1.0f, 1.0f, 1.0f);
diff --git a/Engine/Source/Runtime/Engine/Private/SceneView.cpp b/Engine/Source/Runtime/Engine/Private/SceneView.cpp
--- a/Engine/Source/Runtime/Engine/Private/SceneView.cpp
+++ b/Engine/Source/Runtime/Engine/Private/SceneView.cpp
@@ -146,6 +146,15 @@ static TAutoConsoleVariable<int32> CVarDefaultAutoExposureMethod(
 	TEXT(" 0: Histogram based (requires compute shader, default)\n")
 	TEXT(" 1: Basic AutoExposure"));
 
+static TAutoConsoleVariable<int32> CVarDefaultAutoExposureExtendDefaultLuminanceRange(
+	TEXT("r.DefaultFeature.AutoExposure.ExtendDefaultLuminanceRange"),
+	0,
+	TEXT("Whether the default values for AutoExposure should support an extended range of scene luminance.\n")
+	TEXT("This also change the PostProcessSettings.Exposure.MinBrightness, MaxBrightness, HistogramLogMin and HisogramLogMax\n")
+	TEXT("to be expressed in EV100 values instead of in Luminance and Log2 Luminance.\n")
+	TEXT(" 0: Legacy range (default)\n")
+	TEXT(" 1: Extended range"));
+
 static TAutoConsoleVariable<int32> CVarDefaultMotionBlur(
 	TEXT("r.DefaultFeature.MotionBlur"),
 	1,
diff --git a/Engine/Source/Runtime/Engine/Public/SceneManagement.h b/Engine/Source/Runtime/Engine/Public/SceneManagement.h
--- a/Engine/Source/Runtime/Engine/Public/SceneManagement.h
+++ b/Engine/Source/Runtime/Engine/Public/SceneManagement.h
@@ -195,6 +195,9 @@ public:
 	//
 	virtual uint32 GetFrameIndexMod8() const = 0;
 
+	/** Returns the current PreExposure value. PreExposure is a custom scale applied to the scene color to prevent buffer overflow. */
+	virtual float GetPreExposure() const = 0;
+
 	/** 
 	 * returns the occlusion frame counter 
 	 */
diff --git a/Engine/Source/Runtime/RenderCore/Public/RendererInterface.h b/Engine/Source/Runtime/RenderCore/Public/RendererInterface.h
--- a/Engine/Source/Runtime/RenderCore/Public/RendererInterface.h
+++ b/Engine/Source/Runtime/RenderCore/Public/RendererInterface.h
@@ -602,9 +602,10 @@ public:
 		AllowStaticLighting = true;
 		FrameCountAfterRenderingCommandSend = 0;
 		RequestTickSinceCreation = 0;
+		PreExposure = 1;
 	}
 
-	void SetRequestData(FVector2D SrcViewportUV, int32 TargetBufferIndex, int32 ViewUniqueId, int32 GBufferFormat, bool StaticLightingEnable)
+	void SetRequestData(FVector2D SrcViewportUV, int32 TargetBufferIndex, int32 ViewUniqueId, int32 GBufferFormat, bool StaticLightingEnable, float InPreExposure)
 	{
 		SourceViewportUV = SrcViewportUV;
 		BufferIndex = TargetBufferIndex;
@@ -613,6 +614,7 @@ public:
 		ViewId = ViewUniqueId;
 		GBufferPrecision = GBufferFormat;
 		AllowStaticLighting = StaticLightingEnable;
+		PreExposure = InPreExposure;
 		FrameCountAfterRenderingCommandSend = 0;
 		RequestTickSinceCreation = 0;
 	}
@@ -634,6 +636,7 @@ public:
 	//GPU state at capture time
 	int32 GBufferPrecision;
 	bool AllowStaticLighting;
+	float PreExposure;
 };
 
 #define VIRTUALTEXTURESPACE_MAXLAYERS 4
diff --git a/Engine/Source/Runtime/Renderer/Private/PostProcess/PostProcessEyeAdaptation.cpp b/Engine/Source/Runtime/Renderer/Private/PostProcess/PostProcessEyeAdaptation.cpp
--- a/Engine/Source/Runtime/Renderer/Private/PostProcess/PostProcessEyeAdaptation.cpp
+++ b/Engine/Source/Runtime/Renderer/Private/PostProcess/PostProcessEyeAdaptation.cpp
@@ -44,6 +44,16 @@ TAutoConsoleVariable<float> CVarEyeAdaptationFocus(
 	TEXT(">0: Center focus, 1 is a good number (default)"),
 	ECVF_Scalability | ECVF_RenderThreadSafe);
 
+FORCEINLINE float EV100ToLuminance(float EV100)
+{
+	return 1.2 * FMath::Pow(2.0f, EV100);
+}
+
+FORCEINLINE float EV100ToLog2(float EV100)
+{
+	return EV100 + 0.263f; // Where .263 is log2(1.2)
+}
+
 /**
  *   Shared functionality used in computing the eye-adaptation parameters
  *   Compute the parameters used for eye-adaptation.  These will default to values
@@ -51,6 +61,9 @@ TAutoConsoleVariable<float> CVarEyeAdaptationFocus(
  */
 inline static void ComputeEyeAdaptationValues(const ERHIFeatureLevel::Type MinFeatureLevel, const FViewInfo& View, FVector4 Out[EYE_ADAPTATION_PARAMS_SIZE])
 {
+	static const auto VarDefaultAutoExposureExtendDefaultLuminanceRange = IConsoleManager::Get().FindTConsoleVariableDataInt(TEXT("r.DefaultFeature.AutoExposure.ExtendDefaultLuminanceRange"));
+	const bool bExtendedLuminanceRange = VarDefaultAutoExposureExtendDefaultLuminanceRange->GetValueOnRenderThread() == 1;
+
 	const FPostProcessSettings& Settings = View.FinalPostProcessSettings;
 	const FEngineShowFlags& EngineShowFlags = View.Family->EngineShowFlags;
 
@@ -74,8 +87,8 @@ inline static void ComputeEyeAdaptationValues(const ERHIFeatureLevel::Type MinFe
 	float AverageLuminanceScale = Settings.AutoExposureMethod == EAutoExposureMethod::AEM_Basic ? (1.f / CalibrationConstant) : 1.f;
 
 	// example min/max: -8 .. 4   means a range from 1/256 to 4  pow(2,-8) .. pow(2,4)
-	float HistogramLogMin = Settings.HistogramLogMin;
-	float HistogramLogMax = Settings.HistogramLogMax;
+	float HistogramLogMin = bExtendedLuminanceRange ? EV100ToLog2(Settings.HistogramLogMin) : Settings.HistogramLogMin;
+	float HistogramLogMax = bExtendedLuminanceRange ? EV100ToLog2(Settings.HistogramLogMax) : Settings.HistogramLogMax;
 	HistogramLogMin = FMath::Min<float>(HistogramLogMin, HistogramLogMax - 1);
 
 	// Eye adaptation is disabled except for highend right now because the histogram is not computed.
@@ -87,12 +100,20 @@ inline static void ComputeEyeAdaptationValues(const ERHIFeatureLevel::Type MinFe
 									View.Family->ExposureSettings.FixedEV100 : 
 									FMath::Log2(FMath::Square(Settings.DepthOfFieldFstop) * Settings.CameraShutterSpeed * 100 / FMath::Max(1.f, Settings.CameraISO));
 
-			MinAverageLuminance = MaxAverageLuminance = 1.2f * FMath::Pow(2.0f, FixedEV100);
+			MinAverageLuminance = MaxAverageLuminance = EV100ToLuminance(FixedEV100);
 		}
 		else
 		{
-			MinAverageLuminance = Settings.AutoExposureMinBrightness;
-			MaxAverageLuminance = Settings.AutoExposureMaxBrightness;
+			if (bExtendedLuminanceRange)
+			{
+				MinAverageLuminance = EV100ToLuminance(Settings.AutoExposureMinBrightness);
+				MaxAverageLuminance = EV100ToLuminance(Settings.AutoExposureMaxBrightness);
+			}
+			else
+			{
+				MinAverageLuminance = Settings.AutoExposureMinBrightness;
+				MaxAverageLuminance = Settings.AutoExposureMaxBrightness;
+			}
 		}
 	}
 
@@ -414,7 +435,7 @@ void FSceneViewState::UpdatePreExposure(FViewInfo& View)
 	{
 		const FSceneViewFamily& ViewFamily = *View.Family;
 
-		if (!IsRichView(ViewFamily) && !ViewFamily.EngineShowFlags.VisualizeHDR && !ViewFamily.EngineShowFlags.VisualizeBloom && ViewFamily.EngineShowFlags.PostProcessing && ViewFamily.bResolveScene)
+		if (!IsRichView(ViewFamily) /*&& !ViewFamily.EngineShowFlags.VisualizeHDR */&& !ViewFamily.EngineShowFlags.VisualizeBloom && ViewFamily.EngineShowFlags.PostProcessing && ViewFamily.bResolveScene)
 		{
 			const float PreExposureOverride = CVarEyeAdaptationPreExposureOverride.GetValueOnRenderThread();
 			const float LastExposure = View.GetLastEyeAdaptationExposure();
diff --git a/Engine/Source/Runtime/Renderer/Private/PostProcess/PostProcessTemporalAA.cpp b/Engine/Source/Runtime/Renderer/Private/PostProcess/PostProcessTemporalAA.cpp
--- a/Engine/Source/Runtime/Renderer/Private/PostProcess/PostProcessTemporalAA.cpp
+++ b/Engine/Source/Runtime/Renderer/Private/PostProcess/PostProcessTemporalAA.cpp
@@ -78,7 +78,7 @@ public:
 	FShaderParameter HistoryBufferSize;
 	FShaderParameter HistoryBufferUVMinMax;
 	FShaderParameter MaxViewportUVAndSvPositionToViewportUV;
-	FShaderParameter OneOverHistoryPreExposure;
+	FShaderParameter PreExposureSettings;
 	FShaderParameter ViewportUVToInputBufferUV;
 
 	void Bind(const FShader::CompiledShaderInitializerType& Initializer)
@@ -100,7 +100,7 @@ public:
 		HistoryBufferSize.Bind(ParameterMap, TEXT("HistoryBufferSize"));
 		HistoryBufferUVMinMax.Bind(ParameterMap, TEXT("HistoryBufferUVMinMax"));
 		MaxViewportUVAndSvPositionToViewportUV.Bind(ParameterMap, TEXT("MaxViewportUVAndSvPositionToViewportUV"));
-		OneOverHistoryPreExposure.Bind(ParameterMap, TEXT("OneOverHistoryPreExposure"));
+		PreExposureSettings.Bind(ParameterMap, TEXT("PreExposureSettings"));
 		ViewportUVToInputBufferUV.Bind(ParameterMap, TEXT("ViewportUVToInputBufferUV"));
 	}
 
@@ -108,7 +108,7 @@ public:
 	{
 		Ar << PostprocessParameter << SceneTextureParameters ;
 		Ar << SampleWeights << PlusWeights << DitherScale << VelocityScaling << CurrentFrameWeight << ScreenPosAbsMax << ScreenPosToHistoryBufferUV << HistoryBuffer[0] << HistoryBuffer[1] << HistoryBufferSampler[0] << HistoryBufferSampler[1] << HistoryBufferSize << HistoryBufferUVMinMax;
-		Ar << MaxViewportUVAndSvPositionToViewportUV << OneOverHistoryPreExposure << ViewportUVToInputBufferUV;
+		Ar << MaxViewportUVAndSvPositionToViewportUV << PreExposureSettings << ViewportUVToInputBufferUV;
 	}
 
 	template <typename TRHICmdList, typename TShaderRHIParamRef>
@@ -259,8 +259,17 @@ public:
 			SetShaderValue(RHICmdList, ShaderRHI, MaxViewportUVAndSvPositionToViewportUV, MaxViewportUVAndSvPositionToViewportUVValue);
 		}
 
-		const float OneOverHistoryPreExposureValue = 1.f / FMath::Max<float>(SMALL_NUMBER, InputHistory.IsValid() ? InputHistory.SceneColorPreExposure : Context.View.PreExposure);
-		SetShaderValue(RHICmdList, ShaderRHI, OneOverHistoryPreExposure, OneOverHistoryPreExposureValue);
+		// Pre-exposure, One over Pre-exposure, History pre-exposure, History one over pre-exposure.
+		// DOF settings must preserve scene color range.
+		FVector4 PreExposureSettingsValue(1.f, 1.f, 1.f, 1.f);
+		if (PassParameters.Pass == ETAAPassConfig::Main)
+		{
+			PreExposureSettingsValue.X = Context.View.PreExposure;
+			PreExposureSettingsValue.Y = 1.f / FMath::Max<float>(SMALL_NUMBER, Context.View.PreExposure);
+			PreExposureSettingsValue.Z = InputHistory.IsValid() ? InputHistory.SceneColorPreExposure : Context.View.PreExposure;
+			PreExposureSettingsValue.W  = 1.f / FMath::Max<float>(SMALL_NUMBER, PreExposureSettingsValue.Z);
+		}
+		SetShaderValue(RHICmdList, ShaderRHI, PreExposureSettings, PreExposureSettingsValue);
 
 		{
 			float InvSizeX = 1.0f / float(SrcSize.X);
diff --git a/Engine/Source/Runtime/Renderer/Private/PostProcess/PostProcessVisualizeHDR.cpp b/Engine/Source/Runtime/Renderer/Private/PostProcess/PostProcessVisualizeHDR.cpp
--- a/Engine/Source/Runtime/Renderer/Private/PostProcess/PostProcessVisualizeHDR.cpp
+++ b/Engine/Source/Runtime/Renderer/Private/PostProcess/PostProcessVisualizeHDR.cpp
@@ -197,6 +197,16 @@ FString LogToString(float LogValue)
 	}
 }
 
+FORCEINLINE float LuminanceToEV100(float Luminance)
+{
+	return FMath::Log2(Luminance / 1.2f);
+}
+
+FORCEINLINE float Log2ToEV100(float Log2)
+{
+	return Log2 - 0.263f; // Where .263 is log2(1.2)
+}
+
 void FRCPassPostProcessVisualizeHDR::Process(FRenderingCompositePassContext& Context)
 {
 	SCOPED_DRAW_EVENT(Context.RHICmdList, PostProcessVisualizeHDR);
@@ -208,9 +218,13 @@ void FRCPassPostProcessVisualizeHDR::Process(FRenderingCompositePassContext& Con
 		return;
 	}
 
+	static const auto VarDefaultAutoExposureExtendDefaultLuminanceRange = IConsoleManager::Get().FindTConsoleVariableDataInt(TEXT("r.DefaultFeature.AutoExposure.ExtendDefaultLuminanceRange"));
+	const bool bExtendedLuminanceRange = VarDefaultAutoExposureExtendDefaultLuminanceRange->GetValueOnRenderThread() == 1;
+
 	const FViewInfo& View = Context.View;
 	const FViewInfo& ViewInfo = Context.View;
 	const FSceneViewFamily& ViewFamily = *(View.Family);
+	const EAutoExposureMethod AutoExposureMethod = GetAutoExposureMethod(ViewInfo);
 
 	const FSceneRenderTargetItem& DestRenderTarget = PassOutputs[0].RequestSurface(Context);
 
@@ -262,10 +276,9 @@ void FRCPassPostProcessVisualizeHDR::Process(FRenderingCompositePassContext& Con
 
 	FString Line;
 
-	Line = FString::Printf(TEXT("HDR Histogram (Logarithmic, max of RGB)"));
+	Line = FString::Printf(TEXT("HDR Histogram (EV100, max of RGB)"));
 	Canvas.DrawShadowedString( X, Y += YStep, *Line, GetStatsFont(), FLinearColor(1, 1, 1));
 	
-	
 	Y += 160;
 
 	float MinX = DestRect.Min.X + 64 + 10;
@@ -276,15 +289,18 @@ void FRCPassPostProcessVisualizeHDR::Process(FRenderingCompositePassContext& Con
 	{
 		int XAdd = (int)(i * SizeX / 4);
 		float HistogramPosition =  i / 4.0f;
-		float LogValue = FMath::Lerp(View.FinalPostProcessSettings.HistogramLogMin, View.FinalPostProcessSettings.HistogramLogMax, HistogramPosition);
+		float EV100Value = FMath::Lerp(View.FinalPostProcessSettings.HistogramLogMin, View.FinalPostProcessSettings.HistogramLogMax, HistogramPosition);
+		if (!bExtendedLuminanceRange)
+		{	
+			// In this case the post process settings are actually Log2 values.
+			EV100Value = Log2ToEV100(EV100Value);
+		}
 
-		Line = FString::Printf(TEXT("%.2g"), LogValue);
-		Canvas.DrawShadowedString( MinX + XAdd - 5, MaxY, *Line, GetStatsFont(), FLinearColor(1, 0.3f, 0.3f));
-		Line = LogToString(LogValue);
-		Canvas.DrawShadowedString( MinX + XAdd - 5, MaxY + YStep, *Line, GetStatsFont(), FLinearColor(0.3f, 0.3f, 1));
+		Line = FString::Printf(TEXT("%.2g"), EV100Value);
+		Canvas.DrawShadowedString( MinX + XAdd - 5, MaxY + YStep, *Line, GetStatsFont(), FLinearColor(1, 0.3f, 0.3f));
 	}
 	Y += 3 * YStep;
-	if (GetAutoExposureMethod(ViewInfo) == EAutoExposureMethod::AEM_Basic)
+	if (AutoExposureMethod == EAutoExposureMethod::AEM_Basic)
 	{
 		Line = FString::Printf(TEXT("Basic"));
 	} 
@@ -299,31 +315,45 @@ void FRCPassPostProcessVisualizeHDR::Process(FRenderingCompositePassContext& Con
 	Canvas.DrawShadowedString( X, Y += YStep, TEXT("Percent Low/High:"), GetStatsFont(), FLinearColor(1, 1, 1));
 	Canvas.DrawShadowedString( X + ColumnWidth, Y, *Line, GetStatsFont(), FLinearColor(1, 1, 1));
 
-	Line = FString::Printf(TEXT("%g .. %g"), View.FinalPostProcessSettings.AutoExposureMinBrightness, View.FinalPostProcessSettings.AutoExposureMaxBrightness);
-	Canvas.DrawShadowedString( X, Y += YStep, TEXT("Brightness Min/Max:"), GetStatsFont(), FLinearColor(1, 1, 1));
+	if (bExtendedLuminanceRange)
+	{
+		Line = FString::Printf(TEXT("%.1f .. %.1f"), View.FinalPostProcessSettings.AutoExposureMinBrightness, View.FinalPostProcessSettings.AutoExposureMaxBrightness);
+	}
+	else
+	{
+		Line = FString::Printf(TEXT("%.1f .. %.1f"), LuminanceToEV100(View.FinalPostProcessSettings.AutoExposureMinBrightness), LuminanceToEV100(View.FinalPostProcessSettings.AutoExposureMaxBrightness));
+	}
+	Canvas.DrawShadowedString( X, Y += YStep, TEXT("EV100 Min/Max"), GetStatsFont(), FLinearColor(1, 1, 1));
 	Canvas.DrawShadowedString( X + ColumnWidth, Y, *Line, GetStatsFont(), FLinearColor(0.3f, 0.3f, 1));
 
 	Line = FString::Printf(TEXT("%g / %g"), View.FinalPostProcessSettings.AutoExposureSpeedUp, View.FinalPostProcessSettings.AutoExposureSpeedDown);
 	Canvas.DrawShadowedString( X, Y += YStep, TEXT("Speed Up/Down:"), GetStatsFont(), FLinearColor(1, 1, 1));
 	Canvas.DrawShadowedString( X + ColumnWidth, Y, *Line, GetStatsFont(), FLinearColor(1, 1, 1));
 
-	Line = FString::Printf(TEXT("%g"), View.FinalPostProcessSettings.AutoExposureBias);
-	Canvas.DrawShadowedString( X, Y += YStep, TEXT("Exposure Bias: "), GetStatsFont(), FLinearColor(1, 1, 1));
+	Line = FString::Printf(TEXT("%.2g"), View.FinalPostProcessSettings.AutoExposureBias);
+	Canvas.DrawShadowedString( X, Y += YStep, TEXT("Exposure Compensation: "), GetStatsFont(), FLinearColor(1, 1, 1));
 	Canvas.DrawShadowedString( X + ColumnWidth, Y, *Line, GetStatsFont(), FLinearColor(1, 0.3f, 0.3f));
 
-	Line = FString::Printf(TEXT("%g"), View.FinalPostProcessSettings.AutoExposureCalibrationConstant);
-	Canvas.DrawShadowedString( X, Y += YStep, TEXT("Calibration Constant: "), GetStatsFont(), FLinearColor(1, 1, 1));
-	Canvas.DrawShadowedString( X + ColumnWidth, Y, *Line, GetStatsFont(), FLinearColor(1, 0.3f, 0.3f));
+	if (AutoExposureMethod == AEM_Basic || AutoExposureMethod == AEM_Histogram)
+	{
+		Line = FString::Printf(TEXT("%g%%"), AutoExposureMethod == AEM_Basic ? View.FinalPostProcessSettings.AutoExposureCalibrationConstant : 100.f);
+		Canvas.DrawShadowedString( X, Y += YStep, TEXT("Calibration Constant: "), GetStatsFont(), FLinearColor(1, 1, 1));
+		Canvas.DrawShadowedString( X + ColumnWidth, Y, *Line, GetStatsFont(), FLinearColor(1, 0.3f, 0.3f));
+	}
 
-	Line = FString::Printf(TEXT("%g .. %g (log2)"),
-		View.FinalPostProcessSettings.HistogramLogMin, View.FinalPostProcessSettings.HistogramLogMax);
-	Canvas.DrawShadowedString( X, Y += YStep, TEXT("Log Min/Max:"), GetStatsFont(), FLinearColor(1, 1, 1));
-	Canvas.DrawShadowedString( X + ColumnWidth, Y, *Line, GetStatsFont(), FLinearColor(1, 0.3f, 0.3f));
-	Line = FString::Printf(TEXT("%s .. %s (Value)"),
-		*LogToString(View.FinalPostProcessSettings.HistogramLogMin), *LogToString(View.FinalPostProcessSettings.HistogramLogMax));
-	Canvas.DrawShadowedString( X + ColumnWidth, Y+= YStep, *Line, GetStatsFont(), FLinearColor(0.3f, 0.3f, 1));
+	if (bExtendedLuminanceRange)
+	{
+		Line = FString::Printf(TEXT("%.1f .. %.1f"), View.FinalPostProcessSettings.HistogramLogMin, View.FinalPostProcessSettings.HistogramLogMax);
+	}
+	else
+	{
+		Line = FString::Printf(TEXT("%.1f .. %.1f"), Log2ToEV100(View.FinalPostProcessSettings.HistogramLogMin), Log2ToEV100(View.FinalPostProcessSettings.HistogramLogMax));
+	}
+
+	Canvas.DrawShadowedString( X, Y += YStep, TEXT("Histogram EV100 Min/Max:"), GetStatsFont(), FLinearColor(1, 1, 1));
+	Canvas.DrawShadowedString( X + ColumnWidth, Y, *Line, GetStatsFont(), FLinearColor(0.3f, 0.3f, 1));
 
-	if (GetAutoExposureMethod(ViewInfo) == EAutoExposureMethod::AEM_Basic)
+	if (AutoExposureMethod == EAutoExposureMethod::AEM_Basic)
 	{
 		Line = FString::Printf(TEXT("%g"), GetBasicAutoExposureFocus());
 		Canvas.DrawShadowedString(X, Y += YStep, TEXT("Weighting Focus: "), GetStatsFont(), FLinearColor(1, 1, 1));
diff --git a/Engine/Source/Runtime/Renderer/Private/PostProcess/PostProcessing.cpp b/Engine/Source/Runtime/Renderer/Private/PostProcess/PostProcessing.cpp
--- a/Engine/Source/Runtime/Renderer/Private/PostProcess/PostProcessing.cpp
+++ b/Engine/Source/Runtime/Renderer/Private/PostProcess/PostProcessing.cpp
@@ -1431,7 +1431,7 @@ void FPostProcessing::Process(FRHICommandListImmediate& RHICmdList, const FViewI
 		bool bVisualizeBloom = View.Family->EngineShowFlags.VisualizeBloom && FeatureLevel >= ERHIFeatureLevel::SM4;
 		bool bVisualizeMotionBlur = View.Family->EngineShowFlags.VisualizeMotionBlur && FeatureLevel >= ERHIFeatureLevel::SM4;
 
-		if(bVisualizeHDR || bVisualizeBloom || bVisualizeMotionBlur)
+		if(bVisualizeBloom || bVisualizeMotionBlur)
 		{
 			bAllowTonemapper = false;
 		}
@@ -2003,7 +2003,7 @@ void FPostProcessing::Process(FRHICommandListImmediate& RHICmdList, const FViewI
 			Context.FinalOutput = FRenderingCompositeOutputRef(Node);
 		}
 
-		if(View.Family->EngineShowFlags.VisualizeLPV && !View.Family->EngineShowFlags.VisualizeHDR)
+		if(View.Family->EngineShowFlags.VisualizeLPV)
 		{
 			ensureMsgf(!bUnscaledFinalOutput, TEXT("Should not unscale final output multiple times."));
 			bUnscaledFinalOutput = true;
diff --git a/Engine/Source/Runtime/Renderer/Private/ScenePrivate.h b/Engine/Source/Runtime/Renderer/Private/ScenePrivate.h
--- a/Engine/Source/Runtime/Renderer/Private/ScenePrivate.h
+++ b/Engine/Source/Runtime/Renderer/Private/ScenePrivate.h
@@ -1270,6 +1270,12 @@ public:
 		MIDUsedCount = 0;
 	}
 
+	/** Returns the current PreExposure value. PreExposure is a custom scale applied to the scene color to prevent buffer overflow. */
+	virtual float GetPreExposure() const override
+	{
+		return PreExposure;
+	}
+
 	// Note: OnStartPostProcessing() needs to be called each frame for each view
 	virtual UMaterialInstanceDynamic* GetReusableMID(class UMaterialInterface* InSource) override
 	{		
