Pixel Inspector:
- Pixel inspector now gets the right render target pixel format via scene view extension.
- Pixel Inspector HDR render target texture was using unsupported pixel format.

diff --git a/Engine/Source/Editor/PixelInspector/Private/PixelInspector.cpp b/Engine/Source/Editor/PixelInspector/Private/PixelInspector.cpp
--- a/Engine/Source/Editor/PixelInspector/Private/PixelInspector.cpp
+++ b/Engine/Source/Editor/PixelInspector/Private/PixelInspector.cpp
@@ -36,6 +36,23 @@
 
 namespace PixelInspector
 {
+	/** 
+	* A helper function that extracts the right scene color texture, untouched, to be used further in post processing. 
+	*/
+	FScreenPassTexture ReturnUntouchedSceneColorForPostProcessing(const FPostProcessMaterialInputs& InOutInputs)
+	{
+		if (InOutInputs.OverrideOutput.IsValid())
+		{
+			return InOutInputs.OverrideOutput;
+		}
+		else
+		{
+			/** We don't want to modify scene texture in any way. We just want it to be passed back onto the next stage. */
+			FScreenPassTexture SceneTexture = const_cast<FScreenPassTexture&>(InOutInputs.Textures[(uint32)EPostProcessMaterialInput::SceneColor]);
+			return SceneTexture;
+		}
+	}
+
 	SPixelInspector::SPixelInspector()
 	{
 		DisplayResult = nullptr;
@@ -587,7 +604,7 @@ namespace PixelInspector
 		//Final color can be in HDR (FloatRGBA) or RGB8 formats so we should rely on scene view extension to tell us which format is being used.
 		Buffer_FinalColor_AnyFormat[LastBufferIndex] = NewObject<UTextureRenderTarget2D>(GetTransientPackage(), TEXT("PixelInspectorBufferFinalColorTarget"), RF_Standalone);
 		Buffer_FinalColor_AnyFormat[LastBufferIndex]->AddToRoot();
-		Buffer_FinalColor_AnyFormat[LastBufferIndex]->InitCustomFormat(FinalColorContextGridSize, FinalColorContextGridSize, PixelInspectorSceneViewExtension->GetPixelFormat(), true);
+		Buffer_FinalColor_AnyFormat[LastBufferIndex]->InitCustomFormat(FinalColorContextGridSize, FinalColorContextGridSize, PixelInspectorSceneViewExtension->GetFinalColorPixelFormat(), true);
 		Buffer_FinalColor_AnyFormat[LastBufferIndex]->ClearColor = FLinearColor::Black;
 		Buffer_FinalColor_AnyFormat[LastBufferIndex]->UpdateResourceImmediate(true);
 		FinalColorRenderTargetResource = Buffer_FinalColor_AnyFormat[LastBufferIndex]->GameThread_GetRenderTargetResource();
@@ -603,14 +620,7 @@ namespace PixelInspector
 		//HDR is in float RGB format
 		Buffer_HDR_Float[LastBufferIndex] = NewObject<UTextureRenderTarget2D>(GetTransientPackage(), TEXT("PixelInspectorBufferHDRTarget"), RF_Standalone);
 		Buffer_HDR_Float[LastBufferIndex]->AddToRoot();
-		if (!bInGameViewMode)
-		{
-			Buffer_HDR_Float[LastBufferIndex]->InitCustomFormat(1, 1, PF_FloatRGBA, true);
-		}
-		else
-		{
-			Buffer_HDR_Float[LastBufferIndex]->InitCustomFormat(1, 1, PF_FloatRGB, true);
-		}
+		Buffer_HDR_Float[LastBufferIndex]->InitCustomFormat(1, 1, PixelInspectorSceneViewExtension->GetHDRPixelFormat(), true);
 		Buffer_HDR_Float[LastBufferIndex]->ClearColor = FLinearColor::Black;
 		Buffer_HDR_Float[LastBufferIndex]->UpdateResourceImmediate(true);
 		HDRRenderTargetResource = Buffer_HDR_Float[LastBufferIndex]->GameThread_GetRenderTargetResource();
@@ -709,7 +719,7 @@ namespace PixelInspector
 					PixelResult.OneOverPreExposure = Request.PreExposure > 0.f ? (1.f / Request.PreExposure) : 1.f;;
 
 					FTextureRenderTargetResource* RTResourceFinalColor = Buffer_FinalColor_AnyFormat[Request.BufferIndex]->GameThread_GetRenderTargetResource();
-					const EPixelFormat FinalColorPixelFormat = PixelInspectorSceneViewExtension->GetPixelFormat();
+					const EPixelFormat FinalColorPixelFormat = PixelInspectorSceneViewExtension->GetFinalColorPixelFormat();
 					if (FinalColorPixelFormat == PF_B8G8R8A8)
 					{
 						TArray<FColor> BufferFinalColorValue;
@@ -751,11 +761,12 @@ namespace PixelInspector
 
 					TArray<FLinearColor> BufferHDRValue;
 					FTextureRenderTargetResource* RTResourceHDR = Buffer_HDR_Float[Request.BufferIndex]->GameThread_GetRenderTargetResource();
+					const EPixelFormat HDRPixelFormat = PixelInspectorSceneViewExtension->GetHDRPixelFormat();
 					if (RTResourceHDR->ReadLinearColorPixels(BufferHDRValue) == false)
 					{
 						BufferHDRValue.Empty();
 					}
-					PixelResult.DecodeHDR(BufferHDRValue);
+					PixelResult.DecodeHDR(BufferHDRValue, HDRPixelFormat == PF_FloatRGBA);
 
 					if (Request.GBufferPrecision == EGBufferFormat::Force8BitsPerChannel)
 					{
@@ -861,6 +872,7 @@ namespace PixelInspector
 		: FSceneViewExtensionBase(AutoRegister)
 	{
 		FinalColorPixelFormat = PF_B8G8R8A8;
+		HDRPixelFormat = PF_FloatRGBA;
 		// Default dislay gamma is hardcoded in tonemapper and is set to 2.2. 
 		// We initialize this value and then get the actual value from ViewFamily.
 		FinalColorGamma = DEFAULT_DISPLAY_GAMMA;
@@ -878,24 +890,26 @@ namespace PixelInspector
 		if (PassId == EPostProcessingPass::FXAA)
 		{
 			InOutPassCallbacks.Add(FAfterPassCallbackDelegate::CreateRaw(this, &FPixelInspectorSceneViewExtension::PostProcessPassAfterFxaa_RenderThread));
+		}		
+		
+		if (PassId == EPostProcessingPass::MotionBlur)
+		{
+			InOutPassCallbacks.Add(FAfterPassCallbackDelegate::CreateRaw(this, &FPixelInspectorSceneViewExtension::PostProcessPassAfterMotionBlur_RenderThread));
 		}
 	}
 
 	FScreenPassTexture FPixelInspectorSceneViewExtension::PostProcessPassAfterFxaa_RenderThread(FRDGBuilder& GraphBuilder, const FSceneView& View, const FPostProcessMaterialInputs& InOutInputs)
 	{
 		FinalColorPixelFormat = InOutInputs.Textures[(uint32)EPostProcessMaterialInput::SceneColor].Texture->Desc.Format;
-
-		if (InOutInputs.OverrideOutput.IsValid())
-		{
-			return InOutInputs.OverrideOutput;
-		}
-		else
-		{
-			/** We don't want to modify scene texture in any way. We just want it to be passed back onto the next stage. */
-			FScreenPassTexture SceneTexture = const_cast<FScreenPassTexture&>(InOutInputs.Textures[(uint32)EPostProcessMaterialInput::SceneColor]);
-			return SceneTexture;
-		}
-
+		// Don't need to modify anything, just return the untouched scene color texture back to post processing.
+		return ReturnUntouchedSceneColorForPostProcessing(InOutInputs);
+	}
+	
+	FScreenPassTexture FPixelInspectorSceneViewExtension::PostProcessPassAfterMotionBlur_RenderThread(FRDGBuilder& GraphBuilder, const FSceneView& View, const FPostProcessMaterialInputs& InOutInputs)
+	{
+		HDRPixelFormat = InOutInputs.Textures[(uint32)EPostProcessMaterialInput::SceneColor].Texture->Desc.Format;
+		// Don't need to modify anything, just return the untouched scene color texture back to post processing.
+		return ReturnUntouchedSceneColorForPostProcessing(InOutInputs);
 	}
 
 };
diff --git a/Engine/Source/Editor/PixelInspector/Private/PixelInspector.h b/Engine/Source/Editor/PixelInspector/Private/PixelInspector.h
--- a/Engine/Source/Editor/PixelInspector/Private/PixelInspector.h
+++ b/Engine/Source/Editor/PixelInspector/Private/PixelInspector.h
@@ -39,11 +39,14 @@ namespace PixelInspector
 
 	public:
 		FScreenPassTexture PostProcessPassAfterFxaa_RenderThread(FRDGBuilder& GraphBuilder, const FSceneView& View, const FPostProcessMaterialInputs& InOutInputs);
-		const EPixelFormat GetPixelFormat() const { return FinalColorPixelFormat; }
+		FScreenPassTexture PostProcessPassAfterMotionBlur_RenderThread(FRDGBuilder& GraphBuilder, const FSceneView& View, const FPostProcessMaterialInputs& InOutInputs);
+		const EPixelFormat GetFinalColorPixelFormat() const { return FinalColorPixelFormat; }
+		const EPixelFormat GetHDRPixelFormat() const { return HDRPixelFormat; }
 		const float GetGamma() const { return FinalColorGamma; }
 		
 	private:
 		EPixelFormat FinalColorPixelFormat;
+		EPixelFormat HDRPixelFormat;
 		float FinalColorGamma;
 	};
 
diff --git a/Engine/Source/Editor/PixelInspector/Private/PixelInspectorResult.cpp b/Engine/Source/Editor/PixelInspector/Private/PixelInspectorResult.cpp
--- a/Engine/Source/Editor/PixelInspector/Private/PixelInspectorResult.cpp
+++ b/Engine/Source/Editor/PixelInspector/Private/PixelInspectorResult.cpp
@@ -81,7 +81,7 @@ namespace PixelInspector
 		}
 		Depth = BufferDepthValue[0].R;
 	}
-	void PixelInspectorResult::DecodeHDR(TArray<FLinearColor> &BufferHDRValue)
+	void PixelInspectorResult::DecodeHDR(TArray<FLinearColor> &BufferHDRValue, bool bHasAlphaChannel)
 	{
 		if (BufferHDRValue.Num() <= 0)
 		{
@@ -91,6 +91,11 @@ namespace PixelInspector
 		}
 		HdrLuminance = BufferHDRValue[0].GetLuminance() * OneOverPreExposure;
 		HdrColor = BufferHDRValue[0] * OneOverPreExposure;
+		if (!bHasAlphaChannel)
+		{
+			//Set the alpha to 1.0 as the default value
+			HdrColor.A = 1.0f;
+		}
 	}
 
 	void PixelInspectorResult::DecodeBufferData(TArray<FColor> &BufferAValue, TArray<FColor> &BufferBCDEValue, bool AllowStaticLighting)
diff --git a/Engine/Source/Editor/PixelInspector/Private/PixelInspectorResult.h b/Engine/Source/Editor/PixelInspector/Private/PixelInspectorResult.h
--- a/Engine/Source/Editor/PixelInspector/Private/PixelInspectorResult.h
+++ b/Engine/Source/Editor/PixelInspector/Private/PixelInspectorResult.h
@@ -136,7 +136,7 @@ namespace PixelInspector
 		void DecodeFinalColor(TArray<FLinearColor> &BufferFinalColorValue, float InGamma, bool bHasAlphaChannel);
 		void DecodeSceneColor(TArray<FLinearColor> &BufferSceneColorValue);
 		void DecodeDepth(TArray<FLinearColor> &BufferDepthValue);
-		void DecodeHDR(TArray<FLinearColor> &BufferHDRValue);
+		void DecodeHDR(TArray<FLinearColor> &BufferHDRValue, bool bHasAlphaChannel);
 
 		void DecodeBufferData(TArray<FColor> &BufferAValue, TArray<FColor> &BufferBCDEValue, bool AllowStaticLighting);
 		void DecodeBufferData(TArray<FLinearColor> &BufferAValue, TArray<FColor> &BufferBCDEValue, bool AllowStaticLighting);
diff --git a/Engine/Source/Runtime/Engine/Private/TextureRenderTarget.cpp b/Engine/Source/Runtime/Engine/Private/TextureRenderTarget.cpp
--- a/Engine/Source/Runtime/Engine/Private/TextureRenderTarget.cpp
+++ b/Engine/Source/Runtime/Engine/Private/TextureRenderTarget.cpp
@@ -74,6 +74,7 @@ bool FTextureRenderTargetResource::IsSupportedFormat( EPixelFormat Format )
 	case PF_A16B16G16R16:
 	case PF_FloatRGB:
 	case PF_FloatRGBA: // for exporting materials to .obj/.mtl
+	case PF_FloatR11G11B10: //Pixel inspector for Reading HDR Color
 	case PF_A2B10G10R10: //Pixel inspector for normal buffer
 	case PF_DepthStencil: //Pixel inspector for depth and stencil buffer
 	case PF_G16:// for heightmaps
