PixelInspector: Added readback support for Pixel Format PF_A2B10G10R10.

diff --git a/Engine/Source/Editor/PixelInspector/Private/PixelInspector.cpp b/Engine/Source/Editor/PixelInspector/Private/PixelInspector.cpp
--- a/Engine/Source/Editor/PixelInspector/Private/PixelInspector.cpp
+++ b/Engine/Source/Editor/PixelInspector/Private/PixelInspector.cpp
@@ -729,14 +729,15 @@ namespace PixelInspector
 						}
 						PixelResult.DecodeFinalColor(BufferFinalColorValue);
 					}
-					else if (FinalColorPixelFormat == PF_FloatRGBA || FinalColorPixelFormat == PF_FloatRGB)
+					else if (FinalColorPixelFormat == PF_FloatRGBA || FinalColorPixelFormat == PF_FloatRGB || FinalColorPixelFormat == PF_A2B10G10R10)
 					{
 						TArray<FLinearColor> BufferFinalColorValueLinear;
 						if (RTResourceFinalColor->ReadLinearColorPixels(BufferFinalColorValueLinear) == false)
 						{
 							BufferFinalColorValueLinear.Empty();
 						}
-						PixelResult.DecodeFinalColor(BufferFinalColorValueLinear, PixelInspectorSceneViewExtension->GetGamma(), FinalColorPixelFormat == PF_FloatRGBA);
+						bool bHasAlphaChannel = FinalColorPixelFormat != PF_FloatRGB;
+						PixelResult.DecodeFinalColor(BufferFinalColorValueLinear, PixelInspectorSceneViewExtension->GetGamma(), bHasAlphaChannel);
 					}
 					
 
