Pixel inspector improvements: renamed scene color properties throughout the system, updated comments, fixed transparency and negative values. Luminance also now respects the working color space.

diff --git a/Engine/Source/Editor/PixelInspector/PixelInspectorModule.Build.cs b/Engine/Source/Editor/PixelInspector/PixelInspectorModule.Build.cs
--- a/Engine/Source/Editor/PixelInspector/PixelInspectorModule.Build.cs
+++ b/Engine/Source/Editor/PixelInspector/PixelInspectorModule.Build.cs
@@ -24,7 +24,8 @@ public class PixelInspectorModule : ModuleRules
 		PrivateDependencyModuleNames.AddRange(
              new string[] {
 					"Engine",
-                    "UnrealEd"
+                    "UnrealEd",
+					"ColorManagement"
 				}
          );
 
diff --git a/Engine/Source/Editor/PixelInspector/Private/PixelInspector.cpp b/Engine/Source/Editor/PixelInspector/Private/PixelInspector.cpp
--- a/Engine/Source/Editor/PixelInspector/Private/PixelInspector.cpp
+++ b/Engine/Source/Editor/PixelInspector/Private/PixelInspector.cpp
@@ -62,10 +62,10 @@ namespace PixelInspector
 
 		Buffer_FinalColor_AnyFormat[0] = nullptr;
 		Buffer_FinalColor_AnyFormat[1] = nullptr;
-		Buffer_SceneColor_Float[0] = nullptr;
-		Buffer_SceneColor_Float[1] = nullptr;
-		Buffer_HDR_Float[0] = nullptr;
-		Buffer_HDR_Float[1] = nullptr;
+		Buffer_SceneColorBeforePost_Float[0] = nullptr;
+		Buffer_SceneColorBeforePost_Float[1] = nullptr;
+		Buffer_SceneColorBeforeToneMap_Float[0] = nullptr;
+		Buffer_SceneColorBeforeToneMap_Float[1] = nullptr;
 		Buffer_Depth_Float[0] = nullptr;
 		Buffer_Depth_Float[1] = nullptr;
 		Buffer_A_Float[0] = nullptr;
@@ -528,17 +528,17 @@ namespace PixelInspector
 			Buffer_FinalColor_AnyFormat[BufferIndex]->RemoveFromRoot();
 			Buffer_FinalColor_AnyFormat[BufferIndex] = nullptr;
 		}
-		if (Buffer_SceneColor_Float[BufferIndex] != nullptr)
+		if (Buffer_SceneColorBeforePost_Float[BufferIndex] != nullptr)
 		{
-			Buffer_SceneColor_Float[BufferIndex]->ClearFlags(RF_Standalone);
-			Buffer_SceneColor_Float[BufferIndex]->RemoveFromRoot();
-			Buffer_SceneColor_Float[BufferIndex] = nullptr;
+			Buffer_SceneColorBeforePost_Float[BufferIndex]->ClearFlags(RF_Standalone);
+			Buffer_SceneColorBeforePost_Float[BufferIndex]->RemoveFromRoot();
+			Buffer_SceneColorBeforePost_Float[BufferIndex] = nullptr;
 		}
-		if (Buffer_HDR_Float[BufferIndex] != nullptr)
+		if (Buffer_SceneColorBeforeToneMap_Float[BufferIndex] != nullptr)
 		{
-			Buffer_HDR_Float[BufferIndex]->ClearFlags(RF_Standalone);
-			Buffer_HDR_Float[BufferIndex]->RemoveFromRoot();
-			Buffer_HDR_Float[BufferIndex] = nullptr;
+			Buffer_SceneColorBeforeToneMap_Float[BufferIndex]->ClearFlags(RF_Standalone);
+			Buffer_SceneColorBeforeToneMap_Float[BufferIndex]->RemoveFromRoot();
+			Buffer_SceneColorBeforeToneMap_Float[BufferIndex] = nullptr;
 		}
 		if (Buffer_Depth_Float[BufferIndex] != nullptr)
 		{
@@ -594,12 +594,12 @@ namespace PixelInspector
 		//Release the old buffer
 		ReleaseBuffers(LastBufferIndex);
 
-		FTextureRenderTargetResource* FinalColorRenderTargetResource = nullptr;
-		FTextureRenderTargetResource* SceneColorRenderTargetResource = nullptr;
-		FTextureRenderTargetResource* HDRRenderTargetResource = nullptr;
-		FTextureRenderTargetResource* DepthRenderTargetResource = nullptr;
-		FTextureRenderTargetResource* BufferARenderTargetResource = nullptr;
-		FTextureRenderTargetResource* BufferBCDEFRenderTargetResource = nullptr;
+		FTextureRenderTargetResource* FinalColorResource = nullptr;
+		FTextureRenderTargetResource* SceneColorBeforePostResource = nullptr;
+		FTextureRenderTargetResource* SceneColorBeforeToneMapResource = nullptr;
+		FTextureRenderTargetResource* DepthResource = nullptr;
+		FTextureRenderTargetResource* BufferAResource = nullptr;
+		FTextureRenderTargetResource* BufferBCDEFResource = nullptr;
 		
 		//Final color can be in HDR (FloatRGBA) or RGB8 formats so we should rely on scene view extension to tell us which format is being used.
 		Buffer_FinalColor_AnyFormat[LastBufferIndex] = NewObject<UTextureRenderTarget2D>(GetTransientPackage(), TEXT("PixelInspectorBufferFinalColorTarget"), RF_Standalone);
@@ -607,23 +607,21 @@ namespace PixelInspector
 		Buffer_FinalColor_AnyFormat[LastBufferIndex]->InitCustomFormat(FinalColorContextGridSize, FinalColorContextGridSize, PixelInspectorSceneViewExtension->GetFinalColorPixelFormat(), true);
 		Buffer_FinalColor_AnyFormat[LastBufferIndex]->ClearColor = FLinearColor::Black;
 		Buffer_FinalColor_AnyFormat[LastBufferIndex]->UpdateResourceImmediate(true);
-		FinalColorRenderTargetResource = Buffer_FinalColor_AnyFormat[LastBufferIndex]->GameThread_GetRenderTargetResource();
-
-		//Scene color is in RGB8 format
-		Buffer_SceneColor_Float[LastBufferIndex] = NewObject<UTextureRenderTarget2D>(GetTransientPackage(), TEXT("PixelInspectorBufferSceneColorTarget"), RF_Standalone);
-		Buffer_SceneColor_Float[LastBufferIndex]->AddToRoot();
-		Buffer_SceneColor_Float[LastBufferIndex]->InitCustomFormat(1, 1, PF_FloatRGBA, true);
-		Buffer_SceneColor_Float[LastBufferIndex]->ClearColor = FLinearColor::Black;
-		Buffer_SceneColor_Float[LastBufferIndex]->UpdateResourceImmediate(true);
-		SceneColorRenderTargetResource = Buffer_SceneColor_Float[LastBufferIndex]->GameThread_GetRenderTargetResource();
-
-		//HDR is in float RGB format
-		Buffer_HDR_Float[LastBufferIndex] = NewObject<UTextureRenderTarget2D>(GetTransientPackage(), TEXT("PixelInspectorBufferHDRTarget"), RF_Standalone);
-		Buffer_HDR_Float[LastBufferIndex]->AddToRoot();
-		Buffer_HDR_Float[LastBufferIndex]->InitCustomFormat(1, 1, PixelInspectorSceneViewExtension->GetHDRPixelFormat(), true);
-		Buffer_HDR_Float[LastBufferIndex]->ClearColor = FLinearColor::Black;
-		Buffer_HDR_Float[LastBufferIndex]->UpdateResourceImmediate(true);
-		HDRRenderTargetResource = Buffer_HDR_Float[LastBufferIndex]->GameThread_GetRenderTargetResource();
+		FinalColorResource = Buffer_FinalColor_AnyFormat[LastBufferIndex]->GameThread_GetRenderTargetResource();
+
+		Buffer_SceneColorBeforePost_Float[LastBufferIndex] = NewObject<UTextureRenderTarget2D>(GetTransientPackage(), TEXT("PixelInspectorBufferSceneColorBeforePostTarget"), RF_Standalone);
+		Buffer_SceneColorBeforePost_Float[LastBufferIndex]->AddToRoot();
+		Buffer_SceneColorBeforePost_Float[LastBufferIndex]->InitCustomFormat(1, 1, PF_FloatRGBA, true);
+		Buffer_SceneColorBeforePost_Float[LastBufferIndex]->ClearColor = FLinearColor::Black;
+		Buffer_SceneColorBeforePost_Float[LastBufferIndex]->UpdateResourceImmediate(true);
+		SceneColorBeforePostResource = Buffer_SceneColorBeforePost_Float[LastBufferIndex]->GameThread_GetRenderTargetResource();
+
+		Buffer_SceneColorBeforeToneMap_Float[LastBufferIndex] = NewObject<UTextureRenderTarget2D>(GetTransientPackage(), TEXT("PixelInspectorBufferSceneColorBeforeTonemapTarget"), RF_Standalone);
+		Buffer_SceneColorBeforeToneMap_Float[LastBufferIndex]->AddToRoot();
+		Buffer_SceneColorBeforeToneMap_Float[LastBufferIndex]->InitCustomFormat(1, 1, PixelInspectorSceneViewExtension->GetHDRPixelFormat(), true);
+		Buffer_SceneColorBeforeToneMap_Float[LastBufferIndex]->ClearColor = FLinearColor::Black;
+		Buffer_SceneColorBeforeToneMap_Float[LastBufferIndex]->UpdateResourceImmediate(true);
+		SceneColorBeforeToneMapResource = Buffer_SceneColorBeforeToneMap_Float[LastBufferIndex]->GameThread_GetRenderTargetResource();
 
 		//TODO support Non render buffer to be able to read the depth stencil
 /*		Buffer_Depth_Float[LastBufferIndex] = NewObject<UTextureRenderTarget2D>(GetTransientPackage(), TEXT("PixelInspectorBufferDepthTarget"), RF_Standalone);
@@ -643,14 +641,14 @@ namespace PixelInspector
 			Buffer_A_RGB8[LastBufferIndex]->InitCustomFormat(1, 1, PF_B8G8R8A8, true);
 			Buffer_A_RGB8[LastBufferIndex]->ClearColor = FLinearColor::Black;
 			Buffer_A_RGB8[LastBufferIndex]->UpdateResourceImmediate(true);
-			BufferARenderTargetResource = Buffer_A_RGB8[LastBufferIndex]->GameThread_GetRenderTargetResource();
+			BufferAResource = Buffer_A_RGB8[LastBufferIndex]->GameThread_GetRenderTargetResource();
 
 			Buffer_BCDEF_RGB8[LastBufferIndex] = NewObject<UTextureRenderTarget2D>(GetTransientPackage(), TEXT("PixelInspectorBufferBTarget"), RF_Standalone );
 			Buffer_BCDEF_RGB8[LastBufferIndex]->AddToRoot();
 			Buffer_BCDEF_RGB8[LastBufferIndex]->InitCustomFormat(4, 1, PF_B8G8R8A8, true);
 			Buffer_BCDEF_RGB8[LastBufferIndex]->ClearColor = FLinearColor::Black;
 			Buffer_BCDEF_RGB8[LastBufferIndex]->UpdateResourceImmediate(true);
-			BufferBCDEFRenderTargetResource = Buffer_BCDEF_RGB8[LastBufferIndex]->GameThread_GetRenderTargetResource();
+			BufferBCDEFResource = Buffer_BCDEF_RGB8[LastBufferIndex]->GameThread_GetRenderTargetResource();
 		}
 		else if(GBufferFormat == EGBufferFormat::Default)
 		{
@@ -660,7 +658,7 @@ namespace PixelInspector
 			Buffer_A_RGB10[LastBufferIndex]->InitCustomFormat(1, 1, PF_A2B10G10R10, true);
 			Buffer_A_RGB10[LastBufferIndex]->ClearColor = FLinearColor::Black;
 			Buffer_A_RGB10[LastBufferIndex]->UpdateResourceImmediate(true);
-			BufferARenderTargetResource = Buffer_A_RGB10[LastBufferIndex]->GameThread_GetRenderTargetResource();
+			BufferAResource = Buffer_A_RGB10[LastBufferIndex]->GameThread_GetRenderTargetResource();
 
 			//Default is PF_B8G8R8A8
 			Buffer_BCDEF_RGB8[LastBufferIndex] = NewObject<UTextureRenderTarget2D>(GetTransientPackage(), TEXT("PixelInspectorBufferBTarget"), RF_Standalone );
@@ -668,7 +666,7 @@ namespace PixelInspector
 			Buffer_BCDEF_RGB8[LastBufferIndex]->InitCustomFormat(4, 1, PF_B8G8R8A8, true);
 			Buffer_BCDEF_RGB8[LastBufferIndex]->ClearColor = FLinearColor::Black;
 			Buffer_BCDEF_RGB8[LastBufferIndex]->UpdateResourceImmediate(true);
-			BufferBCDEFRenderTargetResource = Buffer_BCDEF_RGB8[LastBufferIndex]->GameThread_GetRenderTargetResource();
+			BufferBCDEFResource = Buffer_BCDEF_RGB8[LastBufferIndex]->GameThread_GetRenderTargetResource();
 		}
 		else if (GBufferFormat == EGBufferFormat::HighPrecisionNormals || GBufferFormat == EGBufferFormat::Force16BitsPerChannel)
 		{
@@ -678,21 +676,21 @@ namespace PixelInspector
 			Buffer_A_Float[LastBufferIndex]->InitCustomFormat(1, 1, PF_FloatRGBA, true);
 			Buffer_A_Float[LastBufferIndex]->ClearColor = FLinearColor::Black;
 			Buffer_A_Float[LastBufferIndex]->UpdateResourceImmediate(true);
-			BufferARenderTargetResource = Buffer_A_Float[LastBufferIndex]->GameThread_GetRenderTargetResource();
+			BufferAResource = Buffer_A_Float[LastBufferIndex]->GameThread_GetRenderTargetResource();
 
 			Buffer_BCDEF_Float[LastBufferIndex] = NewObject<UTextureRenderTarget2D>(GetTransientPackage(), TEXT("PixelInspectorBufferBTarget"), RF_Standalone );
 			Buffer_BCDEF_Float[LastBufferIndex]->AddToRoot();
 			Buffer_BCDEF_Float[LastBufferIndex]->InitCustomFormat(4, 1, PF_FloatRGBA, true);
 			Buffer_BCDEF_Float[LastBufferIndex]->ClearColor = FLinearColor::Black;
 			Buffer_BCDEF_Float[LastBufferIndex]->UpdateResourceImmediate(true);
-			BufferBCDEFRenderTargetResource = Buffer_BCDEF_Float[LastBufferIndex]->GameThread_GetRenderTargetResource();
+			BufferBCDEFResource = Buffer_BCDEF_Float[LastBufferIndex]->GameThread_GetRenderTargetResource();
 		}
 		else
 		{
 			checkf(0, TEXT("Unhandled gbuffer format (%i) during pixel inspector initializtion."), GBufferFormat);
 		}	
 		
-		SceneInterface->InitializePixelInspector(FinalColorRenderTargetResource, SceneColorRenderTargetResource, DepthRenderTargetResource, HDRRenderTargetResource, BufferARenderTargetResource, BufferBCDEFRenderTargetResource, LastBufferIndex);
+		SceneInterface->InitializePixelInspector(FinalColorResource, SceneColorBeforePostResource, DepthResource, SceneColorBeforeToneMapResource, BufferAResource, BufferBCDEFResource, LastBufferIndex);
 
 		return LastBufferIndex;
 	}
@@ -720,34 +718,26 @@ namespace PixelInspector
 
 					FTextureRenderTargetResource* RTResourceFinalColor = Buffer_FinalColor_AnyFormat[Request.BufferIndex]->GameThread_GetRenderTargetResource();
 					const EPixelFormat FinalColorPixelFormat = PixelInspectorSceneViewExtension->GetFinalColorPixelFormat();
-					if (FinalColorPixelFormat == PF_B8G8R8A8)
+					const float Gamma = (FinalColorPixelFormat == PF_B8G8R8A8) ? 1.0f : PixelInspectorSceneViewExtension->GetGamma();
+					EPixelFormatChannelFlags ValidPixelChannels = GetPixelFormatValidChannels(FinalColorPixelFormat);
+					bool bHasAlphaChannel = EnumHasAnyFlags(ValidPixelChannels, EPixelFormatChannelFlags::A);
+
+					TArray<FLinearColor> BufferFinalColorValueLinear;
+					if (RTResourceFinalColor->ReadLinearColorPixels(BufferFinalColorValueLinear) == false)
 					{
-						TArray<FColor> BufferFinalColorValue;
-						if (RTResourceFinalColor->ReadPixels(BufferFinalColorValue) == false)
-						{
-							BufferFinalColorValue.Empty();
-						}
-						PixelResult.DecodeFinalColor(BufferFinalColorValue);
-					}
-					else if (FinalColorPixelFormat == PF_FloatRGBA || FinalColorPixelFormat == PF_FloatRGB || FinalColorPixelFormat == PF_A2B10G10R10)
-					{
-						TArray<FLinearColor> BufferFinalColorValueLinear;
-						if (RTResourceFinalColor->ReadLinearColorPixels(BufferFinalColorValueLinear) == false)
-						{
-							BufferFinalColorValueLinear.Empty();
-						}
-						bool bHasAlphaChannel = FinalColorPixelFormat != PF_FloatRGB;
-						PixelResult.DecodeFinalColor(BufferFinalColorValueLinear, PixelInspectorSceneViewExtension->GetGamma(), bHasAlphaChannel);
+						BufferFinalColorValueLinear.Empty();
 					}
+
+					PixelResult.DecodeFinalColor(BufferFinalColorValueLinear, Gamma, bHasAlphaChannel);
 					
 
 					TArray<FLinearColor> BufferSceneColorValue;
-					FTextureRenderTargetResource* RTResourceSceneColor = Buffer_SceneColor_Float[Request.BufferIndex]->GameThread_GetRenderTargetResource();
+					FTextureRenderTargetResource* RTResourceSceneColor = Buffer_SceneColorBeforePost_Float[Request.BufferIndex]->GameThread_GetRenderTargetResource();
 					if (RTResourceSceneColor->ReadLinearColorPixels(BufferSceneColorValue) == false)
 					{
 						BufferSceneColorValue.Empty();
 					}
-					PixelResult.DecodeSceneColor(BufferSceneColorValue);
+					PixelResult.DecodeSceneColorBeforePostProcessing(BufferSceneColorValue);
 
 					if (Buffer_Depth_Float[Request.BufferIndex] != nullptr)
 					{
@@ -760,14 +750,17 @@ namespace PixelInspector
 						PixelResult.DecodeDepth(BufferDepthValue);
 					}
 
-					TArray<FLinearColor> BufferHDRValue;
-					FTextureRenderTargetResource* RTResourceHDR = Buffer_HDR_Float[Request.BufferIndex]->GameThread_GetRenderTargetResource();
-					const EPixelFormat HDRPixelFormat = PixelInspectorSceneViewExtension->GetHDRPixelFormat();
-					if (RTResourceHDR->ReadLinearColorPixels(BufferHDRValue) == false)
+					TArray<FLinearColor> BufferSceneColorBeforeToneMapValue;
+					FTextureRenderTargetResource* RTResourceSceneColorBeforeTonemap = Buffer_SceneColorBeforeToneMap_Float[Request.BufferIndex]->GameThread_GetRenderTargetResource();
+					if (RTResourceSceneColorBeforeTonemap->ReadLinearColorPixels(BufferSceneColorBeforeToneMapValue) == false)
 					{
-						BufferHDRValue.Empty();
+						BufferSceneColorBeforeToneMapValue.Empty();
 					}
-					PixelResult.DecodeHDR(BufferHDRValue, HDRPixelFormat == PF_FloatRGBA);
+
+					const EPixelFormat HDRPixelFormat = PixelInspectorSceneViewExtension->GetHDRPixelFormat();
+					ValidPixelChannels = GetPixelFormatValidChannels(HDRPixelFormat);
+					bHasAlphaChannel = EnumHasAnyFlags(ValidPixelChannels, EPixelFormatChannelFlags::A);
+					PixelResult.DecodeSceneColorBeforeToneMap(BufferSceneColorBeforeToneMapValue, bHasAlphaChannel);
 
 					if (Request.GBufferPrecision == EGBufferFormat::Force8BitsPerChannel)
 					{
diff --git a/Engine/Source/Editor/PixelInspector/Private/PixelInspector.h b/Engine/Source/Editor/PixelInspector/Private/PixelInspector.h
--- a/Engine/Source/Editor/PixelInspector/Private/PixelInspector.h
+++ b/Engine/Source/Editor/PixelInspector/Private/PixelInspector.h
@@ -176,10 +176,9 @@ namespace PixelInspector
 		UTextureRenderTarget2D* Buffer_FinalColor_AnyFormat[2];
 		//Depth Buffer
 		UTextureRenderTarget2D* Buffer_Depth_Float[2];
-		//SceneColor Buffer
-		UTextureRenderTarget2D* Buffer_SceneColor_Float[2];
-		//HDR Buffer
-		UTextureRenderTarget2D* Buffer_HDR_Float[2];
+		//SceneColor Buffers
+		UTextureRenderTarget2D* Buffer_SceneColorBeforePost_Float[2];
+		UTextureRenderTarget2D* Buffer_SceneColorBeforeToneMap_Float[2];
 		//GBufferA RenderTarget
 		UTextureRenderTarget2D* Buffer_A_Float[2];
 		UTextureRenderTarget2D* Buffer_A_RGB8[2];
diff --git a/Engine/Source/Editor/PixelInspector/Private/PixelInspectorResult.cpp b/Engine/Source/Editor/PixelInspector/Private/PixelInspectorResult.cpp
--- a/Engine/Source/Editor/PixelInspector/Private/PixelInspectorResult.cpp
+++ b/Engine/Source/Editor/PixelInspector/Private/PixelInspectorResult.cpp
@@ -2,6 +2,7 @@
 
 #include "PixelInspectorResult.h"
 
+#include "ColorSpace.h"
 #include "HAL/PlatformCrt.h"
 #include "Math/Float16.h"
 #include "Math/Float16Color.h"
@@ -12,26 +13,6 @@
 
 namespace PixelInspector
 {
-	void PixelInspectorResult::DecodeFinalColor(TArray<FColor> &BufferFinalColorValue)
-	{
-		FinalColor.Empty();
-		if (BufferFinalColorValue.Num() <= 0)
-		{
-			//Initialize to black
-			for (int i = 0; i < FinalColorContextGridSize*FinalColorContextGridSize; ++i)
-			{
-				FinalColor.Add(FColor::Green);
-			}
-			return;
-		}
-		for (FColor ReadBackColor : BufferFinalColorValue)
-		{
-			ReadBackColor.A = 255;
-			//Set the color in linear
-			FLinearColor LinearColor(ReadBackColor);
-			FinalColor.Add(LinearColor);
-		}
-	}
 	void PixelInspectorResult::DecodeFinalColor(TArray<FLinearColor>& BufferFinalColorValue, float InGamma, bool bHasAlphaChannel)
 	{
 		FinalColor.Empty();
@@ -46,11 +27,11 @@ namespace PixelInspector
 		}
 		for (FLinearColor ReadBackColor : BufferFinalColorValue)
 		{
-			//Set the color in linear
+			//Set the color in linear, with sign/abs functions to prevent NaNs on negative values.
 			FLinearColor FinalColorTemp;
-			FinalColorTemp.R = FMath::Pow(ReadBackColor.R, InGamma);
-			FinalColorTemp.G = FMath::Pow(ReadBackColor.G, InGamma);
-			FinalColorTemp.B = FMath::Pow(ReadBackColor.B, InGamma);
+			FinalColorTemp.R = FMath::Sign(ReadBackColor.R) * FMath::Pow(FMath::Abs(ReadBackColor.R), InGamma);
+			FinalColorTemp.G = FMath::Sign(ReadBackColor.G) * FMath::Pow(FMath::Abs(ReadBackColor.G), InGamma);
+			FinalColorTemp.B = FMath::Sign(ReadBackColor.B) * FMath::Pow(FMath::Abs(ReadBackColor.B), InGamma);
 
 			if (bHasAlphaChannel)
 			{
@@ -64,16 +45,15 @@ namespace PixelInspector
 			FinalColor.Add(FinalColorTemp);		
 		}
 	}
-	void PixelInspectorResult::DecodeSceneColor(TArray<FLinearColor> &BufferSceneColorValue)
+	void PixelInspectorResult::DecodeSceneColorBeforePostProcessing(TArray<FLinearColor> &BufferSceneColorValue)
 	{
 		if (BufferSceneColorValue.Num() <= 0)
 		{
-			SceneColor = FLinearColor::Black;
+			SceneColorBeforePostProcessing = FLinearColor::Transparent;
 			return;
 		}
-		SceneColor = BufferSceneColorValue[0] * OneOverPreExposure;
-		//Set the alpha to 1.0 as the default value
-		SceneColor.A = 1.0f;
+		SceneColorBeforePostProcessing = BufferSceneColorValue[0] * OneOverPreExposure;
+		SceneColorBeforePostProcessing.A = 1.0 - BufferSceneColorValue[0].A;
 	}
 	void PixelInspectorResult::DecodeDepth(TArray<FLinearColor> &BufferDepthValue)
 	{
@@ -85,21 +65,28 @@ namespace PixelInspector
 		}
 		Depth = BufferDepthValue[0].R;
 	}
-	void PixelInspectorResult::DecodeHDR(TArray<FLinearColor> &BufferHDRValue, bool bHasAlphaChannel)
+	void PixelInspectorResult::DecodeSceneColorBeforeToneMap(TArray<FLinearColor>& BufferSceneColorValue, bool bHasAlphaChannel)
 	{
-		if (BufferHDRValue.Num() <= 0)
+		if (BufferSceneColorValue.Num() <= 0)
 		{
-			HdrLuminance = 0.0f;
-			HdrColor = FLinearColor::Black;
+			SceneColorBeforeTonemap = FLinearColor::Transparent;
+			LuminanceBeforeTonemap = 0.0f;
 			return;
 		}
-		HdrLuminance = BufferHDRValue[0].GetLuminance() * OneOverPreExposure;
-		HdrColor = BufferHDRValue[0] * OneOverPreExposure;
-		if (!bHasAlphaChannel)
+		
+		SceneColorBeforeTonemap = BufferSceneColorValue[0] * OneOverPreExposure;
+
+		if (bHasAlphaChannel)
+		{
+			// invert alpha
+			SceneColorBeforeTonemap.A = 1. - BufferSceneColorValue[0].A;
+		}
+		else
 		{
-			//Set the alpha to 1.0 as the default value
-			HdrColor.A = 1.0f;
+			SceneColorBeforeTonemap.A = 1.;
 		}
+
+		LuminanceBeforeTonemap = UE::Color::FColorSpace::GetWorking().GetLuminance(SceneColorBeforeTonemap);
 	}
 
 	void PixelInspectorResult::DecodeBufferData(TArray<FColor> &BufferAValue, TArray<FColor> &BufferBCDEValue, bool AllowStaticLighting)
diff --git a/Engine/Source/Editor/PixelInspector/Private/PixelInspectorResult.h b/Engine/Source/Editor/PixelInspector/Private/PixelInspectorResult.h
--- a/Engine/Source/Editor/PixelInspector/Private/PixelInspectorResult.h
+++ b/Engine/Source/Editor/PixelInspector/Private/PixelInspectorResult.h
@@ -42,7 +42,7 @@ namespace PixelInspector
 			Depth = 0.0f;
 			WorldPosition = FVector(0.0f);
 
-			HdrLuminance = 0.0f;
+			LuminanceBeforeTonemap = 0.0f;
 
 			Normal = FVector(0.0f);
 			PerObjectGBufferData = 0.0f;
@@ -83,18 +83,15 @@ namespace PixelInspector
 
 		//////////////////////////////////////////////////////////////////////////
 		// Scene color
-		FLinearColor SceneColor;
+		FLinearColor SceneColorBeforePostProcessing;
+		FLinearColor SceneColorBeforeTonemap;
+		float LuminanceBeforeTonemap;
 
 		//////////////////////////////////////////////////////////////////////////
 		// Depth and world position
 		float Depth;
 		FVector WorldPosition;
 
-		//////////////////////////////////////////////////////////////////////////
-		// HDR Values
-		float HdrLuminance;
-		FLinearColor HdrColor;
-
 		//////////////////////////////////////////////////////////////////////////
 		//Buffers value
 		FVector Normal; //GBufferA RGB
@@ -138,12 +135,11 @@ namespace PixelInspector
 		float IrisMask;
 		float IrisDistance;
 
-		void DecodeFinalColor(TArray<FColor>& BufferFinalColorValue);
 		/** Decodes final color from HDR input. */
 		void DecodeFinalColor(TArray<FLinearColor> &BufferFinalColorValue, float InGamma, bool bHasAlphaChannel);
-		void DecodeSceneColor(TArray<FLinearColor> &BufferSceneColorValue);
+		void DecodeSceneColorBeforePostProcessing(TArray<FLinearColor> &BufferSceneColorValue);
 		void DecodeDepth(TArray<FLinearColor> &BufferDepthValue);
-		void DecodeHDR(TArray<FLinearColor> &BufferHDRValue, bool bHasAlphaChannel);
+		void DecodeSceneColorBeforeToneMap(TArray<FLinearColor>& BufferSceneColorValue, bool bHasAlphaChannel);
 
 		void DecodeBufferData(TArray<FColor> &BufferAValue, TArray<FColor> &BufferBCDEValue, bool AllowStaticLighting);
 		void DecodeBufferData(TArray<FLinearColor> &BufferAValue, TArray<FColor> &BufferBCDEValue, bool AllowStaticLighting);
diff --git a/Engine/Source/Editor/PixelInspector/Private/PixelInspectorView.cpp b/Engine/Source/Editor/PixelInspector/Private/PixelInspectorView.cpp
--- a/Engine/Source/Editor/PixelInspector/Private/PixelInspectorView.cpp
+++ b/Engine/Source/Editor/PixelInspector/Private/PixelInspectorView.cpp
@@ -14,10 +14,10 @@ UPixelInspectorView::UPixelInspectorView(const FObjectInitializer& ObjectInitial
 		FinalColorContext[i] = FLinearColor::Green;
 	}
 	FinalColor = FLinearColor::Green;
-	SceneColor = FLinearColor::Green;
 	PreExposure = 1.f;
-	Luminance = 0.0f;
-	HdrColor = FLinearColor::Black;
+	SceneColorBeforePostProcessing = FLinearColor::Green;
+	SceneColorBeforeTonemap = FLinearColor::Black;
+	LuminanceBeforeTonemap = 0.0f;
 	Normal = FVector(0.0f);
 	PerObjectGBufferData = 0.0f;
 	Metallic = 0.0f;
@@ -63,10 +63,10 @@ void UPixelInspectorView::SetFromResult(PixelInspector::PixelInspectorResult &Re
 		}
 	}
 	
-	SceneColor = Result.SceneColor;
 	PreExposure = Result.PreExposure;
-	Luminance = Result.HdrLuminance;
-	HdrColor = Result.HdrColor;
+	SceneColorBeforePostProcessing = Result.SceneColorBeforePostProcessing;
+	SceneColorBeforeTonemap = Result.SceneColorBeforeTonemap;
+	LuminanceBeforeTonemap = Result.LuminanceBeforeTonemap;
 	Normal = Result.Normal;
 	PerObjectGBufferData = Result.PerObjectGBufferData;
 	Metallic = Result.Metallic;
diff --git a/Engine/Source/Editor/PixelInspector/Private/PixelInspectorView.h b/Engine/Source/Editor/PixelInspector/Private/PixelInspectorView.h
--- a/Engine/Source/Editor/PixelInspector/Private/PixelInspectorView.h
+++ b/Engine/Source/Editor/PixelInspector/Private/PixelInspectorView.h
@@ -16,6 +16,8 @@ namespace PixelInspector { class PixelInspectorResult; }
 
 #define FinalColorContextGridSize 7
 
+// Note: UPixelInspectorView is only ever created as a transient object, so making properties transient is redundant.
+
 UCLASS(HideCategories = Object, MinimalAPI)
 class UPixelInspectorView : public UObject
 {
@@ -23,25 +25,25 @@ class UPixelInspectorView : public UObject
 
 	FLinearColor FinalColorContext[FinalColorContextGridSize*FinalColorContextGridSize];
 
-	/** Final RGBA 8bits Color after tone mapping, default value is black. */
+	/** Final color in display space after tone mapping. */
 	UPROPERTY(VisibleAnywhere, category = FinalColor)
 	FLinearColor FinalColor;
 
-	/** HDR RGB Color. */
-	UPROPERTY(VisibleAnywhere, category = SceneColor)
-	FLinearColor SceneColor;
-
-	/** HDR Luminance. */
+	/** Amount of exposure applied to Scene Color due to auto-exposure and/or exposure compensation. */
 	UPROPERTY(VisibleAnywhere, category = SceneColor)
 	float PreExposure;
 
-	/** HDR Luminance. */
-	UPROPERTY(VisibleAnywhere, category = HDR)
-	float Luminance;
+	/** HDR scene-linear color in working color space, before post-processing. Previously called Scene Color. */
+	UPROPERTY(VisibleAnywhere, category = SceneColor, meta = (DisplayName = "Scene Color Before Post-Processing"))
+	FLinearColor SceneColorBeforePostProcessing;
 
-	/** HDR RGB Color. */
-	UPROPERTY(VisibleAnywhere, category = HDR)
-	FLinearColor HdrColor;
+	/** HDR scene-linear color in working color space, before tone mapping. Previously called Hdr Color. */
+	UPROPERTY(VisibleAnywhere, category = SceneColor)
+	FLinearColor SceneColorBeforeTonemap;
+
+	/** Luminance of the Scene Color Before Tonemap. */
+	UPROPERTY(VisibleAnywhere, category = SceneColor)
+	float LuminanceBeforeTonemap;
 
 	/** From the GBufferA RGB Channels. */
 	UPROPERTY(VisibleAnywhere, category = GBufferA)
