Post Process:  Support for format conversion in ReturnUntouchedSceneColorForPostProcessing where necessary.  Solves related crashes for the systems that use that function (Pixel Inspector, nDisplay, PPM Chain Graph, Media Capture).  Also added fix for Pixel Inspector Scene View Extension so it only runs logic to cache gamma and pixel format for the view that actually has pixel inspection active, avoiding bugs where other random scene renderers overwrite those values, and unnecessary copies.

#rnx

diff --git a/Engine/Source/Editor/PixelInspector/Private/PixelInspector.cpp b/Engine/Source/Editor/PixelInspector/Private/PixelInspector.cpp
--- a/Engine/Source/Editor/PixelInspector/Private/PixelInspector.cpp
+++ b/Engine/Source/Editor/PixelInspector/Private/PixelInspector.cpp
@@ -852,21 +852,37 @@ namespace PixelInspector
 
 	void FPixelInspectorSceneViewExtension::BeginRenderViewFamily(FSceneViewFamily& InViewFamily)
 	{
-		const float DisplayGamma = InViewFamily.RenderTarget->GetDisplayGamma();
-		// We need to apply gamma to the final color if the output is in HDR.
-		FinalColorGamma = (InViewFamily.EngineShowFlags.Tonemapper == 0) ? DEFAULT_DISPLAY_GAMMA : DisplayGamma;
+		bool bAnyUsePixelInspector = false;
+		for (const FSceneView* View : InViewFamily.Views)
+		{
+			if (View->bUsePixelInspector)
+			{
+				bAnyUsePixelInspector = true;
+				break;
+			}
+		}
+
+		if (bAnyUsePixelInspector)
+		{
+			const float DisplayGamma = InViewFamily.RenderTarget->GetDisplayGamma();
+			// We need to apply gamma to the final color if the output is in HDR.
+			FinalColorGamma = (InViewFamily.EngineShowFlags.Tonemapper == 0) ? DEFAULT_DISPLAY_GAMMA : DisplayGamma;
+		}
 	}
 
 	void FPixelInspectorSceneViewExtension::SubscribeToPostProcessingPass(EPostProcessingPass PassId, const FSceneView& View, FAfterPassCallbackDelegateArray& InOutPassCallbacks, bool bIsPassEnabled)
 	{
-		if (PassId == EPostProcessingPass::FXAA)
+		if (View.bUsePixelInspector)
 		{
-			InOutPassCallbacks.Add(FAfterPassCallbackDelegate::CreateRaw(this, &FPixelInspectorSceneViewExtension::PostProcessPassAfterFxaa_RenderThread));
-		}		
-		
-		if (PassId == EPostProcessingPass::MotionBlur)
-		{
-			InOutPassCallbacks.Add(FAfterPassCallbackDelegate::CreateRaw(this, &FPixelInspectorSceneViewExtension::PostProcessPassAfterMotionBlur_RenderThread));
+			if (PassId == EPostProcessingPass::FXAA)
+			{
+				InOutPassCallbacks.Add(FAfterPassCallbackDelegate::CreateRaw(this, &FPixelInspectorSceneViewExtension::PostProcessPassAfterFxaa_RenderThread));
+			}
+
+			if (PassId == EPostProcessingPass::MotionBlur)
+			{
+				InOutPassCallbacks.Add(FAfterPassCallbackDelegate::CreateRaw(this, &FPixelInspectorSceneViewExtension::PostProcessPassAfterMotionBlur_RenderThread));
+			}
 		}
 	}
 
diff --git a/Engine/Source/Editor/UnrealEd/Private/EditorViewportClient.cpp b/Engine/Source/Editor/UnrealEd/Private/EditorViewportClient.cpp
--- a/Engine/Source/Editor/UnrealEd/Private/EditorViewportClient.cpp
+++ b/Engine/Source/Editor/UnrealEd/Private/EditorViewportClient.cpp
@@ -4120,10 +4120,8 @@ void FEditorViewportClient::SetupViewForRendering(FSceneViewFamily& ViewFamily,
 	GetMeshEdgesViewSettings(View).Opacity = WireframeOpacity;
 	
 	//Look if the pixel inspector tool is on
-	View.bUsePixelInspector = false;
 	FPixelInspectorModule& PixelInspectorModule = FModuleManager::LoadModuleChecked<FPixelInspectorModule>(TEXT("PixelInspectorModule"));
 	bool IsInspectorActive = PixelInspectorModule.IsPixelInspectorEnable();
-	View.bUsePixelInspector = IsInspectorActive;
 	FIntPoint InspectViewportPos = FIntPoint(-1, -1);
 	if (IsInspectorActive)
 	{
@@ -4146,6 +4144,7 @@ void FEditorViewportClient::SetupViewForRendering(FSceneViewFamily& ViewFamily,
 			PixelInspectorModule.SetCoordinatePosition(InspectViewportPos, false);
 		}
 	}
+	View.bUsePixelInspector = IsInspectorActive;
 
 	if (IsInspectorActive)
 	{
diff --git a/Engine/Source/Runtime/Renderer/Private/PostProcess/PostProcessMaterial.cpp b/Engine/Source/Runtime/Renderer/Private/PostProcess/PostProcessMaterial.cpp
--- a/Engine/Source/Runtime/Renderer/Private/PostProcess/PostProcessMaterial.cpp
+++ b/Engine/Source/Runtime/Renderer/Private/PostProcess/PostProcessMaterial.cpp
@@ -1464,3 +1464,20 @@ FScreenPassTexture AddHighResolutionScreenshotMaskPass(
 
 	return Output;
 }
+
+FScreenPassTexture FPostProcessMaterialInputs::ReturnUntouchedSceneColorForPostProcessing(FRDGBuilder& GraphBuilder) const
+{
+	FScreenPassTextureSlice& SceneColorSlice = const_cast<FScreenPassTextureSlice&>(Textures[(uint32)EPostProcessMaterialInput::SceneColor]);
+
+	// Support format conversions here, to handle the case where the output is the final render target, and happens to
+	// be a different format than the intermediate render targets.
+	if (OverrideOutput.IsValid() && OverrideOutput.Texture->Desc.Format != SceneColorSlice.TextureSRV->GetParent()->Desc.Format)
+	{
+		AddDrawTexturePass(GraphBuilder, FScreenPassViewInfo(), SceneColorSlice, OverrideOutput);
+		return OverrideOutput;
+	}
+	else
+	{
+		return FScreenPassTexture::CopyFromSlice(GraphBuilder, SceneColorSlice, OverrideOutput);
+	}
+}
diff --git a/Engine/Source/Runtime/Renderer/Public/PostProcess/PostProcessMaterialInputs.h b/Engine/Source/Runtime/Renderer/Public/PostProcess/PostProcessMaterialInputs.h
--- a/Engine/Source/Runtime/Renderer/Public/PostProcess/PostProcessMaterialInputs.h
+++ b/Engine/Source/Runtime/Renderer/Public/PostProcess/PostProcessMaterialInputs.h
@@ -118,10 +118,7 @@ struct FPostProcessMaterialInputs
 	/**
 	 * Returns the input scene color as a 2D texture output. This WILL perform a GPU copy if the override output is active or the input scene color was a 2D texture array slice.
 	 */
-	inline FScreenPassTexture ReturnUntouchedSceneColorForPostProcessing(FRDGBuilder& GraphBuilder) const
-	{
-		return FScreenPassTexture::CopyFromSlice(GraphBuilder, const_cast<FScreenPassTextureSlice&>(Textures[(uint32)EPostProcessMaterialInput::SceneColor]), OverrideOutput);
-	}
+	RENDERER_API FScreenPassTexture ReturnUntouchedSceneColorForPostProcessing(FRDGBuilder& GraphBuilder) const;
 
 	// [Optional] Render to the specified output. If invalid, a new texture is created and returned.
 	FScreenPassRenderTarget OverrideOutput;
